BuildGraph: Add an option to the Submit task, which allows forcing a change to be submitted even if we have to resolve first. Add a -set:ForceSubmit=true option to the BuildEditorAndTools script that allows using it.

diff --git a/Engine/Build/Graph/Examples/BuildEditorAndTools.xml b/Engine/Build/Graph/Examples/BuildEditorAndTools.xml
--- a/Engine/Build/Graph/Examples/BuildEditorAndTools.xml
+++ b/Engine/Build/Graph/Examples/BuildEditorAndTools.xml
@@ -38,6 +38,7 @@
 	<Option Name="Licensee" Restrict="true|false" DefaultValue="true" Description="Whether to treat the changelist number as being from a licensee Perforce server"/>
 	<Option Name="Versioned" Restrict="true|false" DefaultValue="$(IsBuildMachine)" Description="Whether to embed changelist number into binaries"/>
 	<Option Name="ArchiveStream" Restrict="(?://.*)?" DefaultValue="" Description="Whether to embed changelist number into binaries"/>
+	<Option Name="ForceSubmit" Restrict="true|false" DefaultValue="false" Description="Forces the submit to happen even if another change has been submitted (resolves in favor of local files)"/>
 
 	<EnvVar Name="COMPUTERNAME"/>
 	
@@ -129,7 +130,7 @@
 			<Property Name="ArchiveFile" Value="$(ArchivePerforceDir)\$(EscapedBranch)-Editor.zip"/>
 			<Zip FromDir="$(ArchiveStagingDir)" ZipFile="$(ArchiveFile)"/>
 			<Warning Message="The 'ArchiveStream' argument is not set. Binaries will not be submitted." If="'$(ArchiveStream)' == ''"/>
-			<Submit Description="[CL $(CodeChange)] Updated binaries" Files="$(ArchiveFile)" FileType="binary+FS32" Workspace="$(COMPUTERNAME)_ArchiveForUGS" Stream="$(ArchiveStream)" RootDir="$(ArchivePerforceDir)" If="'$(ArchiveStream)' != ''"/>
+			<Submit Description="[CL $(CodeChange)] Updated binaries" Files="$(ArchiveFile)" FileType="binary+FS32" Workspace="$(COMPUTERNAME)_ArchiveForUGS" Stream="$(ArchiveStream)" RootDir="$(ArchivePerforceDir)" Force="$(ForceSubmit)" If="'$(ArchiveStream)' != ''"/>
 		</Node>
 		
 	</Agent>
diff --git a/Engine/Source/Programs/AutomationTool/BuildGraph/Tasks/SubmitTask.cs b/Engine/Source/Programs/AutomationTool/BuildGraph/Tasks/SubmitTask.cs
--- a/Engine/Source/Programs/AutomationTool/BuildGraph/Tasks/SubmitTask.cs
+++ b/Engine/Source/Programs/AutomationTool/BuildGraph/Tasks/SubmitTask.cs
@@ -57,6 +57,12 @@ namespace AutomationTool.Tasks
 		/// </summary>
 		[TaskParameter(Optional = true)]
 		public bool RevertUnchanged;
+
+		/// <summary>
+		/// Force the submit to happen, even if a resolve is needed (always accept current version)
+		/// </summary>
+		[TaskParameter(Optional = true)]
+		public bool Force;
 	}
 
 	/// <summary>
@@ -144,7 +150,7 @@ namespace AutomationTool.Tasks
 
 				// Submit it
 				int SubmittedCL;
-				SubmitP4.Submit(NewCL, out SubmittedCL);
+				SubmitP4.Submit(NewCL, out SubmittedCL, Force: Parameters.Force);
 				if (SubmittedCL <= 0)
 				{
 					throw new AutomationException("Submit failed.");
