Tag the XGEControlWorker executable as a build product after compiling SCW, to make sure it's included in the UGS zip file.

diff --git a/Engine/Build/Graph/Examples/BuildEditorAndTools.xml b/Engine/Build/Graph/Examples/BuildEditorAndTools.xml
--- a/Engine/Build/Graph/Examples/BuildEditorAndTools.xml
+++ b/Engine/Build/Graph/Examples/BuildEditorAndTools.xml
@@ -54,6 +54,12 @@
 			<Compile Target="ShaderCompileWorker" Platform="Win64" Configuration="Development" Tag="#ToolBinaries"/>
 			<Compile Target="UnrealLightmass" Platform="Win64" Configuration="Development" Tag="#ToolBinaries"/>
 			<Compile Target="UnrealPak" Configuration="Development" Platform="Win64" Tag="#ToolBinaries"/>
+			
+			<!-- 
+				This exe is a copy of ShaderCompileWorker.exe, created as a post-build step. See \Engine\Source\Programs\ShaderCompileWorker\ShaderCompileWorker.Target.cs.
+				It's needed for shader compilation to work with Incredibuild.
+			-->
+			<Tag Files="$(RootDir)\Engine\Binaries\Win64\XGEControlWorker.exe" With="#ToolBinaries"/>
 		</Node>
 		
 		<!-- Compile the editor executable -->
