Include automation modules from platform extensions in archive for UGS.

#rnx

diff --git a/Engine/Build/Graph/Examples/BuildEditorAndTools.xml b/Engine/Build/Graph/Examples/BuildEditorAndTools.xml
--- a/Engine/Build/Graph/Examples/BuildEditorAndTools.xml
+++ b/Engine/Build/Graph/Examples/BuildEditorAndTools.xml
@@ -104,6 +104,7 @@
 
 			<!-- Tag required files from UAT and UBT that will already have been built -->
 			<Tag Files="Engine/Source/Programs/AutomationTool/..." Filter="*.csproj" With="#UAT Projects"/>
+			<Tag Files="Engine/Platforms/..." Filter="*.Automation.csproj" With="#UAT Projects"/>
 			<Tag Files="$(RootDir)/.../Build/Scripts/..." Filter="*.Automation.csproj" With="#UAT Projects"/>
 			<CsCompile Project="#UAT Projects" Configuration="Development" Platform="AnyCPU" Tag="#ArchiveFiles" EnumerateOnly="true"/>
 			
