Change the name of the zip file to EditorTarget rather than just "Editor"
#rnx

diff --git a/Engine/Build/Graph/Examples/BuildEditorAndTools.xml b/Engine/Build/Graph/Examples/BuildEditorAndTools.xml
--- a/Engine/Build/Graph/Examples/BuildEditorAndTools.xml
+++ b/Engine/Build/Graph/Examples/BuildEditorAndTools.xml
@@ -128,7 +128,7 @@
 
 			<!-- Create the zip file and submit it to Perforce -->
 			<Property Name="ArchivePerforceDir" Value="$(ArchiveDir)\Perforce"/>
-			<Property Name="ArchiveFile" Value="$(ArchivePerforceDir)\$(EscapedBranch)-Editor.zip"/>
+			<Property Name="ArchiveFile" Value="$(ArchivePerforceDir)\$(EscapedBranch)-$(EditorTarget).zip"/>
 			<Zip FromDir="$(ArchiveStagingDir)" ZipFile="$(ArchiveFile)"/>
 			<Warning Message="The 'ArchiveStream' argument is not set. Binaries will not be submitted." If="'$(ArchiveStream)' == ''"/>
 			<Submit Description="[CL $(CodeChange)] Updated binaries" Files="$(ArchiveFile)" FileType="binary+FS32" Workspace="$(COMPUTERNAME)_ArchiveForUGS" Stream="$(ArchiveStream)" RootDir="$(ArchivePerforceDir)" Force="$(ForceSubmit)" If="'$(ArchiveStream)' != ''"/>
