Update BuildEditorAndTools.xml to attempt to get the P4CLIENT from the env vars and use that to generate the submit client, otherwise falling back on the old method.

#rnx

diff --git a/Engine/Build/Graph/Examples/BuildEditorAndTools.xml b/Engine/Build/Graph/Examples/BuildEditorAndTools.xml
--- a/Engine/Build/Graph/Examples/BuildEditorAndTools.xml
+++ b/Engine/Build/Graph/Examples/BuildEditorAndTools.xml
@@ -45,6 +45,7 @@
 	<Option Name="SymbolStorePath" DefaultValue="" Description ="Output directory for symbol storage"/>
 
 	<EnvVar Name="COMPUTERNAME"/>
+	<EnvVar Name="P4CLIENT"/>
 	
 	<Agent Name="Default Agent" Type="CompileWin64">
 	
@@ -153,7 +154,9 @@ Platform="AnyCPU"  Properties="EngineDir=$(RootDir)/Engine" Tag="#ArchiveFiles"
 			<Property Name="ArchiveFile" Value="$(ArchivePerforceDir)\$(EscapedBranch)-$(EditorTarget).zip"/>
 			<Zip FromDir="$(ArchiveStagingDir)" ZipFile="$(ArchiveFile)"/>
 			<Warning Message="The 'ArchiveStream' argument is not set. Binaries will not be submitted." If="'$(ArchiveStream)' == ''"/>
-			<Submit Description="[CL $(CodeChange)] Updated binaries" Files="$(ArchiveFile)" FileType="binary+FS32" Workspace="$(COMPUTERNAME)_ArchiveForUGS" Stream="$(ArchiveStream)" RootDir="$(ArchivePerforceDir)" Force="$(ForceSubmit)" If="'$(ArchiveStream)' != ''"/>
+			<Property Name="SubmitClient" Value="$(COMPUTERNAME)_ArchiveForUGS"/>
+			<Property Name="SubmitClient" Value="$(P4CLIENT)_ArchiveForUGS" If="'$(P4CLIENT)' != ''"/>
+			<Submit Description="[CL $(CodeChange)] Updated binaries" Files="$(ArchiveFile)" FileType="binary+FS32" Workspace="$(SubmitClient)" Stream="$(ArchiveStream)" RootDir="$(ArchivePerforceDir)" Force="$(ForceSubmit)" If="'$(ArchiveStream)' != ''"/>
 		</Node>
 		
 	</Agent>
