Checking in new scripts for automating tests on projects.

Frosty, Topaz, ShooterGame, and QAGame have helper scripts set up

diff --git a/Engine/Build/Graph/Tasks/BuildAndTestProject.xml b/Engine/Build/Graph/Tasks/BuildAndTestProject.xml
new file mode 100644
--- /dev/null
+++ b/Engine/Build/Graph/Tasks/BuildAndTestProject.xml
@@ -0,0 +1,296 @@
+<?xml version='1.0' ?>
+<BuildGraph xmlns="http://www.epicgames.com/BuildGraph" xmlns:xsi="http://www.w3.org/2001/XMLSchema-instance" xsi:schemaLocation="http://www.epicgames.com/BuildGraph ../Schema.xsd" >
+	<EnvVar Name="UE_HORDE_JOBID"/>
+	
+	<!-- Project Options -->
+	<Option Name="ProjectName" DefaultValue="" Restrict="" Description="Name of Project"/>
+	<Option Name="TargetName" DefaultValue="" Restrict="" Description="If your project supports both Game and Client you can set this, otherwise one will be selected"/>
+
+	<Option Name="ProjectPath" DefaultValue="" Restrict=".*" Description="Path to the folder that contains the project (should not include the .uproject file!)" />
+	<Option Name="PlatformsToSkipTests" DefaultValue="" Restrict=".*" Description="List of platforms that tests won't be run on" />
+
+	<Option Name="WithBATDefaults" DefaultValue="true" Description="Used by parent scripts to specify defaults when including this script"/>
+
+	<!-- 
+		Allow tests to specify their own defaults by specifying a WithBATDefaults property that is true. The reason 
+		we use an option is we need something to check, and a property in a parent script will always override 
+		an option. Note - this means if a parent script includes this script multiple tiles then it should
+		always WithBATDefaults true/false before doing so, and if false it should supply all these arguments.
+	-->
+	<Do If="$(WithBATDefaults)">
+		<Property Name="DefaultEditorPlatforms" Value="" />
+		<Property Name="DefaultTargetPlatforms" Value="" />
+		<Property Name="DefaultEditorTestList" Value="UE.EditorAutomation(RunTest=Project.Maps.PIE)" />
+		<Property Name="DefaultTargetTestList" Value="UE.TargetAutomation(RunTest=Project.Maps.Cycle)" />
+	</Do>
+
+	<!-- This will define options for all platforms and create lists called EditorPlatforms, RequiredEditorPlatforms, and TargetPlatforms. E.g. TargetPlatforms=Win64+Mac+PS4 -->
+	<Include Script="Inc/PlatformOptions.xml" />
+
+	<!-- Does this project require a dedicated server for testing? (not currently implemented!)-->
+	<Option Name="DedicatedServer" DefaultValue="false" Description="Project uses a dedicated server" />
+
+	<!-- Customizations -->
+	<Option Name="SkipBuild" DefaultValue="false" Description="Whether to skip building" />
+	<Option Name="SkipCook" DefaultValue="false" Description="Whether to skip cooking and packaging" />
+	<Option Name="SkipPackage" DefaultValue="false" Description="Whether to skip packaging" />
+	<Option Name="SkipPublish" DefaultValue="false" Description="Whether to skip cooking, packaging, publishing" />
+	<Option Name="SkipTest" DefaultValue="false" Description="Whether to skip testing" />
+
+	<!-- Defines whether we always run a simple 'does it boot' test before running other tests. This can save a lot of time/repeated errors -->
+	<Option Name="EditorBootTest" DefaultValue="!$(SkipTest)" Description="Run a simple boot test on the editor before other tests"/>
+	<Option Name="TargetBootTest" DefaultValue="!$(SkipTest)" Description="Run a simple boot test on target platforms before other tests"/>
+
+	<Option Name="EditorTestList" DefaultValue="$(DefaultEditorTestList)" Description="Tests to run on the edutor. Test1+Test2 etc"/>
+	<Option Name="TargetTestList" DefaultValue="$(DefaultTargetTestList)" Description="Tests to run on the target platform(s). Test1+Test2 etc"/>
+
+	<!-- Settings files that can be overriden -->
+	<Option Name="PathSettingsFile" DefaultValue="Inc/PathSettings.xml" Description="Path settings to use" />
+	<Option Name="GauntletSettingsFile" DefaultValue="Inc/GauntletSettings.xml" Description="Gauntlet settings to use" />
+	
+	<!-- Some common properties -->
+	<Include Script="Inc/CommonProperties.xml" />
+
+	<!-- Common settings but can be overidden -->
+	<Include Script="$(PathSettingsFile)" />
+	<Include Script="$(GauntletSettingsFile)" />
+
+	<!-- At this point we have all options so calculate a few things -->
+
+	<!-- Validate that the project exists-->
+	<Error Message="ProjectName is not set" If="'$(ProjectName)' == ''" />
+	<Property Name="TargetProject" Value="$(ProjectPath)/$(ProjectName).uproject" />
+	<Error Message="Project file $(TargetProject) does not exist" If="!Exists('$(TargetProject)')" />
+
+	<!-- If the user didn't supply a TargetName then try to auto detect is. This is reliable uness the project has 'Game' in it but the target doesn't -->
+	<Do If="'$(TargetName)' == ''">
+		<!-- 
+			If the user didn't supply a TargetName, try to figure it out based on GameName.Target or GameNameClient.target.
+			Note - if the project is ShooterGame and the target is ShooterClient.txt then this will fail and we'll tell the user -->
+		<Property Name="TargetName" Value="$(ProjectName)Client" If="Exists('$(ProjectPath)/Source/$(ProjectName)Client.Target.cs')" />
+		<Property Name="TargetName" Value="$(ProjectName)" If="Exists('$(ProjectPath)/Source/$(ProjectName).Target.cs')" />
+
+		<Do If="'$(TargetName)' == ''">
+			<Do If="Exists('$(ProjectPath)/Source/$(ProjectName)/$(ProjectName).Build.cs')">
+				<Error Message="The project has source but the name the target couldn't be determined. use -set:TargetName:Foo to ensure things build correctly" />
+			</Do>
+			<Property Name="TargetName" Value="UnrealGame" />
+		</Do>
+	</Do>
+
+	<!-- If skip test is true, add all target platforms to the "do not test" list -->
+	<Do If="$(SkipTest)">
+		<Property Name="PlatformsToSkipTests" Value="$(EditorPlatforms)+$(TargetPlatforms)" />
+	</Do>
+
+	<!-- If skip packaging is true, add all target platforms to the "do not test" list -->
+	<Do If="$(SkipPackage)">
+		<Property Name="PlatformsToSkipPackaging" Value="$(EditorPlatforms)+$(TargetPlatforms)" />
+	</Do>
+
+	<!-- Detect the type of the client based on the target name -->
+	<Property Name="TargetType" Value="Game" />
+	<Property Name="TargetType" Value="Client" If="Contains('$(TargetName)', 'Client')" />
+
+	<!-- We want a display name for a number of jobs which is best if it's the target (e.g. ShooterClient), but not if it's content Only -->
+	<Property Name="DisplayTargetName" Value="$(TargetName)" />
+	<Property Name="DisplayTargetName" Value="$(ProjectName)" If="$(TargetName) == 'UnrealGame'" />
+	
+	<!-- These properties will be filled in with nodes that are generated based on options -->
+	<Property Name="GeneratedToolNodes" Value="" />
+	<Property Name="GeneratedEditorTestNodes" Value="" />
+	<Property Name="GeneratedTargetNodes" Value="" />
+	<Property Name="GeneratedTargetTestNodes" Value="" />
+
+	<!-- Do all editor steps under a single agent as the things that can be done in parallel probably aren't worth the overhead-->
+	<ForEach Name="EditorPlatform" Values="$(RequiredEditorPlatforms)" Separator="+">
+		<!-- Set Agent Types -->
+		<Property Name="HostAgentType" Value="$(EditorPlatform)" />
+		<Property Name="TestAgentType" Value="Test$(EditorPlatform)" />
+
+		<!-- Test nodes for this editor platform -->
+		<Property Name="PlatformEditorTestNodes" Value="" />
+
+		<Agent Name="Build $(ProjectName)Editor $(EditorPlatform)" Type="$(HostAgentType)">
+			<!-- First generate nodes to build the required tools (e.g. UHT, ShaderCompiler, Editor -->
+			<Node Name="Compile $(ProjectName)Tools $(EditorPlatform)">
+				<Compile Target="UnrealHeaderTool" Platform="$(EditorPlatform)" Configuration="Development"/>
+				<Compile Target="ShaderCompileWorker" Platform="$(EditorPlatform)" Configuration="Development"/>
+				<Compile Target="UnrealPak" Platform="$(EditorPlatform)" Configuration="Development"/>
+				<Compile Target="CrashReportClient" Platform="$(EditorPlatform)" Configuration="Shipping"/>
+				<Compile Target="CrashReportClientEditor" Platform="$(EditorPlatform)" Configuration="Shipping"/>
+			</Node>
+			<Node Name="Compile $(ProjectName)Editor $(EditorPlatform)" Requires="Compile $(ProjectName)Tools $(EditorPlatform)">
+				<Compile Target="$(ProjectName)Editor" Platform="$(EditorPlatform)" Configuration="Development" />
+			</Node>
+			<!--Add these nodes to our dependency list -->
+			<Property Name="GeneratedToolNodes" Value="Compile $(ProjectName)Editor $(EditorPlatform); Compile $(ProjectName)Tools $(EditorPlatform)" If="!$(SkipBuild)" />
+		</Agent>
+		<Agent Name="Test $(ProjectName)Editor $(EditorPlatform)" Type="$(TestAgentType)">
+			<!-- Prerequisites for the editor test to run -->
+			<Property Name="TestPrerequisites" Value="$(GeneratedToolNodes)"/>
+			<!-- Declare a boot test for the editor-->
+			<Property Name="PlatformBootTestNodeName" Value="$(ProjectName)Editor $(EditorPlatform) Test=BootTest" />
+			<Do If="'$(AutomationReportOutputDirectory)' != ''">
+				<Property Name="OutputFragmentPath" Value="$(EditorPlatform)Editor/UE.EditorBootTest" />
+				<Property Name="OutputReportPath" Value="$(AutomationReportOutputDirectory)/$(OutputFragmentPath)" />
+				<Property Name="GauntletReportArgs" Value="-logdir=&quot;$(OutputReportPath)&quot; -artifactname=Gauntlet -ReportExportPath=&quot;$(OutputReportPath)&quot;" />
+				<Property Name="GauntletReportArgs" Value="$(GauntletReportArgs) -ReportURL=&quot;$(AutomationReportUrl)/$(OutputFragmentPath)/&quot;" If="'$(AutomationReportUrl)' != ''" />
+			</Do>
+			<Node Name="$(PlatformBootTestNodeName)" Requires="$(GeneratedToolNodes)">
+				<Command Name="RunUnreal" Arguments="-project=$(TargetProject) -build=Editor -test=&quot;UE.EditorBootTest&quot; $(GauntletEditorTestArgs) $(GauntletReportArgs)" />
+			</Node>
+			<!-- Link the boot test in to the prerequisites and generated nodes if its enabled-->
+			<Do If="$(EditorBootTest)">
+				<Property Name="TestPrerequisites" Value="$(TestPrerequisites);$(PlatformBootTestNodeName)" />
+			</Do>
+			<!-- Now generate a test node for each editor test in the list-->
+			<ForEach Name="TestName" Values="$(EditorTestList)" Separator="+">
+			<Do If="'$(AutomationReportOutputDirectory)' != ''">
+				<Property Name="OutputFragmentPath" Value="$(EditorPlatform)/Editor$(TestName)" />
+				<Property Name="OutputReportPath" Value="$(AutomationReportOutputDirectory)/$(OutputFragmentPath)" />
+				<Property Name="GauntletReportArgs" Value="-logdir=&quot;$(OutputReportPath)&quot; -artifactname=Gauntlet -ReportExportPath=&quot;$(OutputReportPath)&quot;" />
+				<Property Name="GauntletReportArgs" Value="$(GauntletReportArgs) -ReportURL=&quot;$(AutomationReportUrl)/$(OutputFragmentPath)/&quot;" If="'$(AutomationReportUrl)' != ''" />
+			</Do>
+				<Property Name="TestNodeName" Value="$(ProjectName)Editor $(EditorPlatform) Test=$(TestName)" />
+				<Node Name="$(TestNodeName)" Requires="$(TestPrerequisites)">
+					<Command Name="RunUnreal" Arguments="-project=$(TargetProject) -build=Editor -test=&quot;$(TestName)&quot; $(GauntletReportArgs) $(GauntletEditorTestArgs)" />
+				</Node>
+				<!-- Add each test node to our generated list -->
+				<Property Name="PlatformEditorTestNodes" Value="$(PlatformEditorTestNodes);$(TestNodeName)"  />
+			</ForEach>
+
+			<!-- Create Labels -->
+			<Label Category="Editors" Name="$(ProjectName) Editor $(EditorPlatform)" Requires="$(GeneratedToolNodes)" />
+
+			<!-- If this editor was requested (as opposed to being required due to target platforms) then link in the test nodes and set up labels -->
+			<Do If="ContainsItem('$(EditorPlatforms)', '$(EditorPlatform)', '+') And !ContainsItem('$(PlatformsToSkipTests)', '$(EditorPlatform)', '+')">
+				<Label Category="Test" Name="BootTest $(ProjectName)Editor $(EditorPlatform)" Requires="$(PlatformBootTestNodeName)" />
+				<Label Category="Test" Name="Test $(ProjectName)Editor $(EditorPlatform)" Requires="$(PlatformEditorTestNodes)" />
+				<Property Name="GeneratedEditorTestNodes" Value="$(GeneratedEditorTestNodes);$(PlatformBootTestNodeName);$(PlatformEditorTestNodes)" />
+			</Do>
+		</Agent>
+	</ForEach>
+
+	<!-- Now, for each platform we build, cook, package, and test -->
+	<ForEach Name="TargetPlatform" Values="$(TargetPlatforms)" Separator="+">
+
+		<!-- Declare type of host required to build & cook this platform-->
+		<Property Name="HostAgentType" Value="Win64" />
+		<Property Name="HostAgentType" Value="Mac" If="'$(TargetPlatform)' == 'Mac' or '$(TargetPlatform)' == 'IOS'"/>
+		<!-- On consoles we can use regualr agents, but locally hosted tests need a Test agent-->
+		<Property Name="TestAgentType" Value="$(HostAgentType)" />
+		<Property Name="TestAgentType" Value="Test$(HostAgentType)" If="'$(TargetPlatform)' == 'Win64' or '$(TargetPlatform)' == 'Mac'"/>
+
+		<!-- Building and cooking will require these nodes from the appropriate host. They should have been built above... -->
+		<Property Name="PlatformToolsNodeName" Value="Compile $(ProjectName)Tools $(HostAgentType)" />
+		<Property Name="PlatformEditorNodeName" Value="Compile $(ProjectName)Editor $(HostAgentType)" />
+		<!-- Define these names up front as we use some of them several times -->
+		<Property Name="PlatformAgentName" Value="$(DisplayTargetName) $(TargetPlatform)" />
+		<Property Name="PlatformCompileNodeName" Value="Compile $(TargetName) $(TargetPlatform) $(TargetConfiguration)" />
+		<Property Name="PlatformCookNodeName" Value="Cook $(DisplayTargetName) $(TargetPlatform) $(TargetConfiguration)" />
+		<Property Name="PlatformPackageNodeName" Value="Package $(DisplayTargetName) $(TargetPlatform) $(TargetConfiguration)" />
+		<Property Name="PlatformPublishNodeName" Value="Publish $(DisplayTargetName) $(TargetPlatform) $(TargetConfiguration)" />
+		<Property Name="PlatformBootTestNodeName" Value="$(DisplayTargetName) $(TargetPlatform) Test=BootTest" />
+
+		<!-- This will be a generated string of all tests for this platform that can be tied to an aggregate -->
+		<Property Name="PlatformTestNodeNames" Value="" />
+		
+		<!-- Set CookPlatform since it isn't always the same as the target -->
+		<Property Name="CookPlatform" Value="$(TargetPlatform)" />
+		<Property Name="CookPlatform" Value="Android_ETC2" If="'$(TargetPlatform)' == 'Android'"/>
+		<Property Name="CookPlatform" Value="Windows" If="'$(TargetPlatform)' == 'Win64'"/>
+		<Property Name="CookPlatform" Value="$(CookPlatform)Client" If="'$(TargetType)' == 'Client'" />
+
+		<!-- Build the client and if necessary a server-->
+		<Agent Name="Build $(PlatformAgentName)" Type="$(HostAgentType)">
+			<Node Name="$(PlatformCompileNodeName)" Requires="$(PlatformToolsNodeName)">
+				<Compile Target="$(TargetName)" Platform="$(TargetPlatform)" Configuration="$(TargetConfiguration)"/>
+				<Do If="'$(DedicatedServer)'" >
+					<Compile Target="$(TargetName)Server" Platform="Win64" Configuration="$(TargetConfiguration)"/>
+				</Do>
+			</Node>
+		</Agent>
+
+		<!-- Cook for the platform. This can run at the same time as the above since we don't need the exe -->
+		<Agent Name="Cook $(PlatformAgentName)" Type="$(HostAgentType)">
+			<Node Name="$(PlatformCookNodeName)" Requires="$(PlatformEditorNodeName)">
+				<Cook Project="$(TargetProject)" Platform="$(CookPlatform)" Arguments="" />				
+			</Node>			
+		</Agent>
+
+		<!-- Now package the exe and cooked data-->
+		<Agent Name="Package $(PlatformAgentName)" Type="$(HostAgentType)">
+			<Node Name="$(PlatformPackageNodeName)" Requires="$(PlatformCompileNodeName);$(PlatformCookNodeName)">
+				<Property Name="BCRArgs" Value="-project=$(TargetProject) -platform=$(TargetPlatform) -configuration=$(TargetConfiguration) -skipbuild -skipcook -pak -stage -archive -archivedirectory=&quot;$(ProjectOutputDirectory)&quot;" />
+				<Property Name="BCRArgs" Value="$(BCRArgs) -client" If="'$(TargetType)' == 'Client'" />
+				<Property Name="BCRArgs" Value="$(BCRArgs) -server" If="'$(DedicatedServer)'" />
+				<Property Name="BCRArgs" Value="$(BCRArgs) -package" If="!ContainsItem('$(PlatformsToSkipPackaging)','$(TargetPlatform)','+')" />		
+				<Command Name="BuildCookRun" Arguments="$(BCRArgs)" />
+			</Node>
+			<Node Name="$(PlatformPublishNodeName)" Requires="$(PlatformPackageNodeName)">
+				<!-- Typically only build machines publish -->
+				<Do If="'$(ProjectPublishPath)' != ''" >
+					<Copy Files="..." From="$(ProjectOutputDirectory)/$(CookPlatform)" To="$(ProjectPublishPath)/$(CookPlatform)"/>
+				</Do>
+			</Node>
+		</Agent>			
+
+		<!-- Test Time! -->		
+		<Agent Name="Test $(PlatformAgentName)" Type="$(TestAgentType)">	
+			<!-- Prerequisites for the test to run -->
+			<Property Name="TestPrerequisites" Value="$(PlatformEditorNodeName)" />
+			<!-- Declare a boot test for the target-->	
+			<Node Name="$(PlatformBootTestNodeName)" After="$(PlatformPublishNodeName)">
+				<Do If="'$(AutomationReportOutputDirectory)' != ''">
+					<Property Name="OutputFragmentPath" Value="$(TargetPlatform)/UE.BootTest" />
+					<Property Name="OutputReportPath" Value="$(AutomationReportOutputDirectory)/$(OutputFragmentPath)" />
+					<Property Name="GauntletReportArgs" Value="-logdir=&quot;$(OutputReportPath)&quot; -artifactname=Gauntlet -ReportExportPath=&quot;$(OutputReportPath)&quot;" />
+					<Property Name="GauntletReportArgs" Value="$(GauntletReportArgs) -ReportURL=&quot;$(AutomationReportUrl)/$(OutputFragmentPath)/&quot;" If="'$(AutomationReportUrl)' != ''" />
+				</Do>
+				<Command Name="RunUnreal" Arguments="-project=$(TargetProject) -test=&quot;UE.BootTest&quot; -platform=$(TargetPlatform) -configuration=$(TargetConfiguration) $(GauntletTargetTestArgs) $(GauntletReportArgs)" />
+			</Node>
+			<!-- Link the boot test in to the prerequisites and generated nodes if its enabled-->
+			<Do If="$(TargetBootTest)">
+				<Property Name="TestPrerequisites" Value="$(TestPrerequisites);$(PlatformBootTestNodeName)" />
+			</Do>
+			<!-- Go through the list of tests -->	
+			<ForEach Name="TestName" Values="$(TargetTestList)" Separator="+">
+				<!-- Run and report on the test-->
+				<Do If="'$(AutomationReportOutputDirectory)' != ''">
+					<Property Name="OutputFragmentPath" Value="$(TargetPlatform)/$(TestName)" />
+					<Property Name="OutputReportPath" Value="$(AutomationReportOutputDirectory)/$(OutputFragmentPath)" />
+					<Property Name="GauntletReportArgs" Value="-logdir=&quot;$(OutputReportPath)&quot; -artifactname=Gauntlet -ReportExportPath=&quot;$(OutputReportPath)&quot;" />
+					<Property Name="GauntletReportArgs" Value="$(GauntletReportArgs) -ReportURL=&quot;$(AutomationReportUrl)/$(OutputFragmentPath)/&quot;" If="'$(AutomationReportUrl)' != ''" />
+				</Do>
+				<Property Name="TestNodeName" Value="$(DisplayTargetName) $(TargetPlatform) Test=$(TestName)" />
+				<!-- Run the test. No 'Requires' as we run on the same machine locally and use the publish folder otherwose -->
+				<Node Name="$(TestNodeName)" Requires="$(TestPrerequisites)">
+					<Command Name="RunUnreal" Arguments="-project=$(TargetProject) -test=&quot;$(TestName)&quot; -platform=$(TargetPlatform) -configuration=$(TargetConfiguration) $(GauntletTargetTestArgs) $(GauntletReportArgs)"/>
+				</Node>
+				<!-- Add it to our list -->
+				<Property Name="PlatformTestNodeNames" Value="$(PlatformTestNodeNames);$(TestNodeName)"/>
+			</ForEach>	
+		</Agent>
+
+		<!-- Declare labels for CIS -->
+		<Label Category="Platforms" Name="$(ProjectName) $(TargetPlatform)" Requires="$(PlatformCompileNodeName)" />
+		<Label Category="Cook" Name="Cook $(DisplayTargetName) $(TargetPlatform)" Requires="$(PlatformCookNodeName)" Exclude="$(PlatformEditorNodeName)" />
+		<Label Category="Package" Name="Pkg $(DisplayTargetName) $(TargetPlatform)" Requires="$(PlatformPackageNodeName);$(PlatformPublishNodeName)" Exclude="$(PlatformCompileNodeName);$(PlatformCookNodeName)"/>
+		<Label Category="Test" Name="BootTest $(DisplayTargetName) $(TargetPlatform)" Requires="$(PlatformBootTestNodeName)" Exclude="$(PlatformEditorNodeName)" />
+		<Label Category="Test" Name="Test $(DisplayTargetName) $(TargetPlatform)" Requires="$(PlatformTestNodeNames)" Exclude="$(PlatformEditorNodeName);$(PlatformBootTestNodeName)"/>
+
+		<!-- Tie all steps into our platform tasks uless excluded -->
+		<Property Name="PlatformTasks" Value="" />
+		<Property Name="PlatformTasks" Value="$(PlatformCompileNodeName)" If="!$(SkipBuild)" />
+		<Property Name="PlatformTasks" Value="$(PlatformTasks);$(PlatformCookNodeName)" If="!$(SkipCook) and !$(SkipPublish)" />		
+		<Property Name="PlatformTasks" Value="$(PlatformTasks);$(PlatformPackageNodeName);$(PlatformPublishNodeName)" If="!$(SkipPublish)" />		
+		<Property Name="PlatformTasks" Value="$(PlatformTasks);$(PlatformBootTestNodeName);$(PlatformTestNodeNames)" If="!ContainsItem('$(PlatformsToSkipTests)','$(TargetPlatform)','+')" />		
+
+		<Property Name="GeneratedTargetNodes" Value="$(GeneratedTargetNodes);$(PlatformTasks)"/>
+	</ForEach>
+
+	<!-- Create an aggregate with all the generated nodes -->
+	<Aggregate Name="BuildAndTest $(ProjectName)" Requires="$(GeneratedToolNodes);$(GeneratedEditorTestNodes);$(GeneratedTargetNodes);$(GeneratedTargetTestNodes)"/>
+
+</BuildGraph>
\ No newline at end of file
diff --git a/Engine/Build/Graph/Tasks/Inc/CommonProperties.xml b/Engine/Build/Graph/Tasks/Inc/CommonProperties.xml
new file mode 100644
--- /dev/null
+++ b/Engine/Build/Graph/Tasks/Inc/CommonProperties.xml
@@ -0,0 +1,23 @@
+<?xml version='1.0' ?>
+<BuildGraph xmlns="http://www.epicgames.com/BuildGraph" xmlns:xsi="http://www.w3.org/2001/XMLSchema-instance" xsi:schemaLocation="http://www.epicgames.com/BuildGraph ../Schema.xsd" >
+
+	<!--
+		This script defines some common properties based on arguments like preflight and project name.
+		
+		InProperty: Change - Change we are building against
+		InProperty: PreflightChange - CL of the change being preflighted
+        InProperty: SetPreflightFlag - Mark this job as a preflight
+		InProperty: ProjectName - Name of the project
+	-->
+
+	<Option Name="PreflightChange" DefaultValue="" Description="Preflight CL number if preflight, empty otherwise"/>
+	<Option Name="SetPreflightFlag" DefaultValue="true" Description="Whether to set the IsPreflight flag to true for preflights. Use with caution: this will cause things like build versions to be set to their non-preflight form."/>
+
+	<Property Name="IsPreflight" Value="false"/>
+	<Property Name="IsPreflight" Value="true" If="'$(PreflightChange)' != '' And $(SetPreflightFlag)"/>
+	<Property Name="PreflightSuffix" Value=""/>
+	<Property Name="PreflightSuffix" Value="-PF-$(PreflightChange)-$(UE_HORDE_JOBID)" If="$(IsPreflight)"/>
+	<Property Name="CL-String" Value="CL-$(Change)$(PreflightSuffix)"/>
+	<Property Name="BuildName" Value="$(ProjectName)-$(CL-String)"/>
+
+</BuildGraph>
diff --git a/Engine/Build/Graph/Tasks/Inc/GauntletSettings.xml b/Engine/Build/Graph/Tasks/Inc/GauntletSettings.xml
new file mode 100644
--- /dev/null
+++ b/Engine/Build/Graph/Tasks/Inc/GauntletSettings.xml
@@ -0,0 +1,40 @@
+<?xml version='1.0' ?>
+<BuildGraph xmlns="http://www.epicgames.com/BuildGraph" xmlns:xsi="http://www.w3.org/2001/XMLSchema-instance" xsi:schemaLocation="http://www.epicgames.com/BuildGraph ../../Schema.xsd" >
+
+	<!--
+		This script defines parameters passed to Gauntlet. It is expected that most licensees will need
+        to create a custom version of this scrupt due to things like the device service url
+
+		This script should be includedEng after CommonProperties.xml and PathSettings.xml where build
+        strings and paths are defines
+		
+		InProperty: ProjectPath - Path where the .uproject is located
+        InProperty: ProjectName - Name of the project (e.g. <ProjectName>.uproject)
+		InProperty: EscapedBranch
+		InProperty: ProjectOutputDirectory - Local pathwhere builds are found
+		InProperty: ProjectPublishPath - Path (if any) builds are published to
+
+        OutProperty: GauntletEditorTestArgs - Arguments to pass to Gauntlet tests executed under the editor
+		OutProperty: GauntletTargetTestArgs - Arguments to pass to Gauntlet tests executed using cooked builds
+	-->
+
+
+	<!-- How long to give Gauntlet before a test is considered timed out (some tests will time out earlier if they detect inactivity)-->
+	<Option Name="GauntletTimeout" DefaultValue="3600" Description="Max running time of a test in Gauntlet"/>
+	
+	<!-- Create the -build argument for Gauntlet based on whether builds are local or on the network-->
+	<Property Name="GauntletBuildArgs" Value="-build=&quot;$(ProjectOutputDirectory)&quot;" />
+	<Property Name="GauntletBuildArgs" Value="-build=&quot;$(ProjectPublishPath)&quot;" If ="'$(ProjectPublishPath)' != ''" />
+
+	<Property Name="GauntletCommonArgs" Value="-maxduration=$(GauntletTimeOut) -unattended -verbose" />
+
+	<Property Name="GauntletDeviceArgs" Value="" />
+	<Property Name="GauntletDeviceArgs" Value="-deviceurl=http://deviceservice.epicgames.net -reboot -cleardevices -sandbox=&quot;$(ProjectName)&quot;" If ="$(IsBuildMachine)" />
+
+	<Property Name="GauntletTargetTestArgs" Value="$(GauntletBuildArgs) $(GauntletCommonArgs) $(GauntletDeviceArgs)" />
+	<Property Name="GauntletEditorTestArgs" Value="-build=editor $(GauntletCommonArgs)" />
+
+    <!-- this is fully generated dynamiclly when needed, it's simply here so it can be used without worrying if its set -->
+    <Property Name="GauntletReportArgs" Value="" />
+
+</BuildGraph>
\ No newline at end of file
diff --git a/Engine/Build/Graph/Tasks/Inc/PathSettings.xml b/Engine/Build/Graph/Tasks/Inc/PathSettings.xml
new file mode 100644
--- /dev/null
+++ b/Engine/Build/Graph/Tasks/Inc/PathSettings.xml
@@ -0,0 +1,63 @@
+<?xml version='1.0' ?>
+<BuildGraph xmlns="http://www.epicgames.com/BuildGraph" xmlns:xsi="http://www.w3.org/2001/XMLSchema-instance" xsi:schemaLocation="http://www.epicgames.com/BuildGraph ../../Schema.xsd" >
+
+	<!--
+		This script defines the local and network paths used for builds. It is expected that most licensees
+		will need to replace this script due to the network paths.
+		
+		On a build machine a network path is always needed to allow tests on other agents retrieve the 
+		build, however a differentiation is made between "temp" storage and "build" storage. At Epic 
+		the former has a life of <1 day whereas build are evicted over time.
+
+		This script should be included after CommonProperties.xml where build strings are defines.
+
+		InOption: PublishBuilds - Enabled/Disable publishing of builds to build storage
+		
+		InProperty: ProjectName
+		InProperty: EscapedBranch
+
+		OutProperty: ProjectOutputDirectory - Local path to package builds to
+		OutProperty: ProjectPublishPath - Path to publish builds to. If empty publishing should be skipped
+	-->
+
+	<!-- Publish builds to the network ?-->
+	<Option Name="PublishBuilds" DefaultValue="false" Description="Publish packaged builds to the builds folder on the netwok"/>
+	<Option Name="PublishReports" DefaultValue="$(PublishBuilds)" Description="Publish packaged builds to the builds folder on the netwok"/>
+
+
+	<!-- Where builds are saved to locally (todo - use ProjectDir?) -->
+	<Property Name="ProjectOutputDirectory" Value="$(RootDir)/LocalBuilds/$(ProjectName)"/>
+	<Property Name="AutomationReportOutputDirectory" Value=""/>
+
+	<!-- Scratch network path that can be used for tests if we aren't publishing. Assumes there's a cleanup script.. -->
+	<Property Name="NetworkTempRoot" Value="/Volumes/Builds/UE5/TempStorage"/>
+	<Property Name="NetworkTempRoot" Value="P:/Builds/UE5/TempStorage" If="'$(HostPlatform)' == 'Win64'"/>
+
+	<!-- Network path that can be used for builds that are published -->
+	<Property Name="NetworkPublishRoot" Value="/Volumes/Builds/UE5/PackagedBuilds"/>
+	<Property Name="NetworkPublishRoot" Value="P:/Builds/UE5/PackagedBuilds" If="'$(HostPlatform)' == 'Win64'"/>
+
+	<!-- Network path that can be used for storing reports and logs -->
+	<Property Name="NetworkReportRoot" Value="/Volumes/Builds/Automation/Reports"/>
+	<Property Name="NetworkReportRoot" Value="P:/Builds/Automation/Reports" If="'$(HostPlatform)' == 'Win64'"/>
+
+	<!-- Declare the publish path, and point to the network if PublishBuilds is true -->
+	<Property Name="ProjectPublishPath" Value="" />
+	<Property Name="ProjectPublishPath" Value="$(NetworkPublishRoot)" If="$(PublishBuilds)"/>
+
+	<!-- Automation can publish reports to a folder that is mapped to a website -->
+	<Do If="$(IsBuildMachine) or $(PublishReports)" >
+		<Property Name="AutomationReportOutputDirectory" Value="$(NetworkReportRoot)"/>
+		<Property Name="AutomationReportOutputDirectory" Value="$(AutomationReportOutputDirectory)/$(EscapedBranch)/$(ProjectName)/$(CL-String)"/>
+		<Property Name="AutomationReportUrl" Value="http://automation.epicgames.net/reports/$(EscapedBranch)/$(ProjectName)/$(CL-String)"/>
+	</Do>
+
+	<!-- Build machines always need to publish somewhere so other agents can run the builds-->
+	<Do If="$(IsBuildMachine) And !$(PublishBuilds)">
+		<Property Name="ProjectPublishPath" Value="$(NetworkTempRoot)/$(EscapedBranch)/$(ProjectName)/$(CL-String)" />
+	</Do>
+
+	<!-- If we have a publish dir now, fully form the path-->
+	<Property Name="ProjectPublishPath" Value="$(NetworkPublishRoot)/$(EscapedBranch)/$(BuildName)" If="'$(ProjectPublishPath)' != ''"/>
+
+</BuildGraph>
\ No newline at end of file
diff --git a/Engine/Build/Graph/Tasks/Inc/PlatformOptions.xml b/Engine/Build/Graph/Tasks/Inc/PlatformOptions.xml
new file mode 100644
--- /dev/null
+++ b/Engine/Build/Graph/Tasks/Inc/PlatformOptions.xml
@@ -0,0 +1,102 @@
+<?xml version='1.0' ?>
+<BuildGraph xmlns="http://www.epicgames.com/BuildGraph" xmlns:xsi="http://www.w3.org/2001/XMLSchema-instance" xsi:schemaLocation="http://www.epicgames.com/BuildGraph ../../Schema.xsd" >
+
+	<!--
+		This script exposes all supported platforms as options and turns the selected platforms
+		into three + separated lists.
+
+		1) EditorPlatforms - these are editor platforms the user requested directly or via WithWin64Editor=true
+		2) TargetPlatforms - these are cooked platforms the user requested directly or via WithPS4=true
+		3) RequiredEditorPlatforms - these are editor platforms that need to be build due to #2, but shouldn't be
+		 used for tests etc because they weren't in #1
+
+		Both 1) and 2) can be specified directly for convenience. E.g. -set:TargetPlatforms=PS5+XSX
+
+		Downstream code can make the distinction of an editor to build but not test by the following 
+		condition being true - ContainsItem('$(EditorName)', '$(EditorPlatforms)', '+')
+
+		If no platforms are found either by being specified directly or via the 'WithFoo' syntax, then 
+		the DefaultEditorPlatforms and DefaultTargetPlatforms properties are returned. This allows scripts
+		to set defaults.
+
+		OutProperty: EditorPlatforms - Editor platforms to build and test, separated by +
+		OutProperty: TargetPlatforms - Target platforms to build and test, separated by +
+		OutProperty: RequiredEditorPlatforms - Editor platforms to build to cook targets      
+	-->
+
+<Option Name="EditorPlatforms" DefaultValue="" Description="Alternative to 'WithFooEditor', a list of editors to build and test. Win64+Mac etc"/>
+<Option Name="TargetPlatforms" DefaultValue="" Description="Alternative to 'WithFoo', a list of Platforms to build and test. PS4+XSX etc"/>
+
+<Option Name="WithWin64Editor" DefaultValue="false" Description="Build and test Win64 Editor"/>
+<Option Name="WithMacEditor" DefaultValue="false" Description="Build and test Mac Editor"/>
+
+<Option Name="WithAndroid" DefaultValue="false" Description="Build and test Android"/>
+<Option Name="WithIOS" DefaultValue="false" Description="Build and test IOS"/>
+<Option Name="WithLumin" DefaultValue="false" Description="Build and test Lumin"/>
+<Option Name="WithMac" DefaultValue="false" Description="Build and test Mac"/>
+<Option Name="WithPS4" DefaultValue="false" Description="Build and test PS4"/>
+<Option Name="WithPS5" DefaultValue="false" Description="Build and test PS5"/>
+<Option Name="WithStadia" DefaultValue="false" Description="Build and test Stadia"/>
+<Option Name="WithSwitch" DefaultValue="false" Description="Build and test Switch"/>
+<Option Name="WithWin64" DefaultValue="false" Description="Build and test Win64"/>
+<Option Name="WithWinGDK" DefaultValue="false" Description="Build and test WinGDK"/>
+<Option Name="WithXboxOneGDK" DefaultValue="false" Description="Build and test XboxOneGDK"/>
+<Option Name="WithXSX" DefaultValue="false" Description="Build and test XSX"/>
+
+<Option Name="TargetConfiguration" DefaultValue="Development" Description="Configuration to build and test"/>
+
+<Do If="'$(EditorPlatforms)' == ''" >
+	<Property Name="EditorPlatforms" Value="" />
+	<Property Name="EditorPlatforms" Value="$(EditorPlatforms)Win64+" If="$(WithWin64Editor)" />
+	<Property Name="EditorPlatforms" Value="$(EditorPlatforms)Mac+" If="$(WithMacEditor)" />
+</Do>
+
+<Do If="'$(TargetPlatforms)' == ''" >
+	<Property Name="TargetPlatforms" Value="" />
+	<Property Name="TargetPlatforms" Value="$(TargetPlatforms)Android+" If="$(WithAndroid)" />
+	<Property Name="TargetPlatforms" Value="$(TargetPlatforms)IOS+" If="$(WithIOS)" />
+	<Property Name="TargetPlatforms" Value="$(TargetPlatforms)Lumin+" If="$(WithLumin)" />
+	<Property Name="TargetPlatforms" Value="$(TargetPlatforms)Mac+" If="$(WithMac)" />
+	<Property Name="TargetPlatforms" Value="$(TargetPlatforms)PS4+" If="$(WithPS4)" />
+	<Property Name="TargetPlatforms" Value="$(TargetPlatforms)PS5+" If="$(WithPS5)" />
+	<Property Name="TargetPlatforms" Value="$(TargetPlatforms)Stadia+" If="$(WithStadia)" />
+	<Property Name="TargetPlatforms" Value="$(TargetPlatforms)Switch+" If="$(WithSwitch)" />
+	<Property Name="TargetPlatforms" Value="$(TargetPlatforms)Win64+" If="$(WithWin64)" />
+	<Property Name="TargetPlatforms" Value="$(TargetPlatforms)WinGDK+" If="$(WithWinGDK)" />
+	<Property Name="TargetPlatforms" Value="$(TargetPlatforms)XboxOneGDK+" If="$(WithXboxOneGDK)" />
+	<Property Name="TargetPlatforms" Value="$(TargetPlatforms)XSX+" If="$(WithXSX)" />
+</Do>
+
+<!-- If still no platforms, pick the defaults that the script should have specified-->
+<Do If="'$(EditorPlatforms)' == ''" >
+	<Property Name="EditorPlatforms" Value="$(DefaultEditorPlatforms)" />
+</Do>
+
+<Do If="'$(TargetPlatforms)' == ''" >
+	<Property Name="TargetPlatforms" Value="$(DefaultTargetPlatforms)" />
+</Do>
+
+<!-- 'None' means don't do any platforms. We need to do this because an empty string means use defaults.. -->
+<Do If="'$(EditorPlatforms)' == 'None'" >
+	<Property Name="EditorPlatforms" Value="" />
+</Do>
+<Do If="'$(TargetPlatforms)' == 'None'" >
+	<Property Name="TargetPlatforms" Value="" />
+</Do>
+
+<Property Name="RequiredEditorPlatforms" Value="$(EditorPlatforms)" />
+
+<ForEach Name="TargetPlatform" Values="$(TargetPlatforms)" Separator="+">
+	<!-- Declare type of host required to build & cook this platform-->
+	<Property Name="HostAgentType" Value="Win64" />
+	<Property Name="HostAgentType" Value="Mac" If="'$(TargetPlatform)' == 'Mac' or '$(TargetPlatform)' == 'IOS'"/>
+
+	<Do If="!ContainsItem('$(RequiredEditorPlatforms)','$(HostAgentType)', '+')">
+		<Property Name="RequiredEditorPlatforms" Value="$(RequiredEditorPlatforms)+$(HostAgentType)" />
+	</Do>
+</ForEach>
+
+<!-- PS5 packaging isn't supported by Gauntlet yet :( )-->
+<Property Name="PlatformsToSkipPackaging" Value="PS5" />
+
+</BuildGraph>
\ No newline at end of file
diff --git a/Engine/Source/Programs/AutomationTool/Gauntlet/Unreal/Automation/UE.BootTest.cs b/Engine/Source/Programs/AutomationTool/Gauntlet/Unreal/Automation/UE.BootTest.cs
--- a/Engine/Source/Programs/AutomationTool/Gauntlet/Unreal/Automation/UE.BootTest.cs
+++ b/Engine/Source/Programs/AutomationTool/Gauntlet/Unreal/Automation/UE.BootTest.cs
@@ -14,7 +14,7 @@ namespace UE
 	/// <summary>
 	/// Test that waits for the client and server to get to the front-end then quits
 	/// </summary>
-	public class BootTest : UnrealTestNode<UnrealGame.UE4TestConfig>
+	public class BootTest : UnrealTestNode<UnrealTestConfiguration>
 	{
 		/// <summary>
 		/// Used to track progress via logging
@@ -44,9 +44,9 @@ namespace UE
 		/// Returns the configuration description for this test
 		/// </summary>
 		/// <returns></returns>
-		public override UnrealGame.UE4TestConfig GetConfiguration()
+		public override UnrealTestConfiguration GetConfiguration()
 		{
-			UnrealGame.UE4TestConfig Config = base.GetConfiguration();
+			UnrealTestConfiguration Config = base.GetConfiguration();
 
 			UnrealTestRole Client = Config.RequireRole(UnrealTargetRole.Client);
 
@@ -75,29 +75,41 @@ namespace UE
 			return true;
 		}
 
+		/// <summary>
+		/// String that we search for to be considered "Booted"
+		/// </summary>
+		/// <returns></returns>
+		protected virtual string GetCompletionString()
+		{
+			return "Bringing up level for play took";
+		}
+
 		/// <summary>
 		/// Called periodically while the test is running to allow code to monitor health.
 		/// </summary>
 		public override void TickTest()
 		{
 			const int kTimeOutDuration = 2;
-			const string kStartupCompleteString = "Bringing up level for play took";
 
 			// run the base class tick;
 			base.TickTest();
 
 			// Get the log of the first client app
-			IAppInstance RunningInstance = this.TestInstance.ClientApps.First();
+			IAppInstance RunningInstance = this.TestInstance.RunningRoles.First().AppInstance;
 
 			UnrealLogParser LogParser = new UnrealLogParser(RunningInstance.StdOut);
 
-			// count how many lines there are in the log to check progress
-			int LogLines = LogParser.Content.Split(new[] { '\n' }, StringSplitOptions.RemoveEmptyEntries).Length;
+			IEnumerable<string> BusyLogLines = LogParser.GetEditorBusyChannels();
+			int BusyLineCount = BusyLogLines.Count();
 
-			if (LogLines > LogLinesLastTick)
+			if (BusyLineCount > LogLinesLastTick)
 			{
 				LastLogTime = DateTime.Now;
+				// log new entries so people have something to look at
+				BusyLogLines.Skip(LogLinesLastTick).ToList().ForEach(S => Log.Info("{0}", S));
+				LogLinesLastTick = BusyLineCount;
 			}
+
 			// Gauntlet will timeout tests based on the -timeout argument, but we have greater insight here so can bail earlier to save
 			// tests from idling on the farm needlessly.
 			if ((DateTime.Now - LastLogTime).TotalMinutes > kTimeOutDuration)
@@ -110,12 +122,17 @@ namespace UE
 			// now see if the game has brought the first world up for play
 			IEnumerable<string> LogWorldLines = LogParser.GetLogChannel("World");
 
-			if (LogWorldLines.Any(L => L.IndexOf(kStartupCompleteString) >= 0))
+			string CompletionString = GetCompletionString();
+
+			if (!string.IsNullOrEmpty(CompletionString))
 			{
-				Log.Info("Found world is ready for play. Ending Test");
-				MarkTestComplete();
-				DidDetectLaunch = true;
-				SetUnrealTestResult(TestResult.Passed);
+				if (LogParser.Content.IndexOf(CompletionString, StringComparison.OrdinalIgnoreCase) > 0)
+				{
+					Log.Info("Found '{0}'. Ending Test", GetCompletionString());
+					MarkTestComplete();
+					DidDetectLaunch = true;
+					SetUnrealTestResult(TestResult.Passed);
+				}
 			}
 		}
 
@@ -174,6 +191,58 @@ namespace UE
 			return base.CreateReport(GetTestResult());
 		}
 	}
+
+	/// <summary>
+	/// Test that verifies the editor boots
+	/// </summary>
+	public class EditorBootTest : BootTest
+	{
+		public EditorBootTest(Gauntlet.UnrealTestContext InContext)
+			: base(InContext)
+		{
+		}
+
+		/// <summary>
+		/// Returns the configuration description for this test
+		/// </summary>
+		/// <returns></returns>
+		public override UnrealTestConfiguration GetConfiguration()
+		{
+			UnrealTestConfiguration Config = base.GetConfiguration();
+			// currently needed as BootTest isn't an abstract class. Can be changed for 4.27
+			Config.ClearRoles();
+			UnrealTestRole EditorRole = Config.RequireRole(UnrealTargetRole.Editor);
+			EditorRole.CommandLineParams.Add("execcmds", "QUIT_EDITOR");
+			return Config;
+		}
+
+		protected override string GetCompletionString()
+		{
+			return null;
+		}
+	}
+
+	/// <summary>
+	/// Test that verifies a target boots
+	/// </summary>
+	public class TargetBootTest : BootTest
+	{
+		public TargetBootTest(Gauntlet.UnrealTestContext InContext)
+			: base(InContext)
+		{
+		}
+
+		/// <summary>
+		/// Returns the configuration description for this test
+		/// </summary>
+		/// <returns></returns>
+		public override UnrealTestConfiguration GetConfiguration()
+		{
+			UnrealTestConfiguration Config = base.GetConfiguration();
+			Config.RequireRole(UnrealTargetRole.Client);
+			return Config;
+		}
+	}
 }
 
 // Provided for backwards compatibility with scripts
diff --git a/Engine/Source/Programs/AutomationTool/Gauntlet/Unreal/Base/Gauntlet.UnrealTestConfiguration.cs b/Engine/Source/Programs/AutomationTool/Gauntlet/Unreal/Base/Gauntlet.UnrealTestConfiguration.cs
--- a/Engine/Source/Programs/AutomationTool/Gauntlet/Unreal/Base/Gauntlet.UnrealTestConfiguration.cs
+++ b/Engine/Source/Programs/AutomationTool/Gauntlet/Unreal/Base/Gauntlet.UnrealTestConfiguration.cs
@@ -669,6 +669,14 @@ namespace Gauntlet
 			return RequireRoles(InRole, null, Count);
 		}
 
+		/// <summary>
+		/// Clears all roles from this config. 
+		/// </summary>
+		public void ClearRoles()
+		{
+			RequiredRoles.Clear();
+		}
+
 		public IEnumerable<UnrealTestRole> RequireRoles(UnrealTargetRole InRole, UnrealTargetPlatform? PlatformOverride, int Count, ERoleModifier roleType = ERoleModifier.None)
 		{
 			if (RequiredRoles.ContainsKey(InRole) == false)
diff --git a/Engine/Source/Programs/AutomationTool/Gauntlet/Unreal/Utils/Gauntlet.UnrealLogParser.cs b/Engine/Source/Programs/AutomationTool/Gauntlet/Unreal/Utils/Gauntlet.UnrealLogParser.cs
--- a/Engine/Source/Programs/AutomationTool/Gauntlet/Unreal/Utils/Gauntlet.UnrealLogParser.cs
+++ b/Engine/Source/Programs/AutomationTool/Gauntlet/Unreal/Utils/Gauntlet.UnrealLogParser.cs
@@ -347,6 +347,15 @@ namespace Gauntlet
 			return Regex.Matches(Content, Pattern).Cast<Match>().Select(M => M.Groups[1].ToString()).ToArray();
 		}
 
+		/// <summary>
+		/// Returns channels that signify the editor doing stuff
+		/// </summary>
+		/// <returns></returns>
+		public IEnumerable<string> GetEditorBusyChannels ()
+		{
+			return GetLogChannels(new string[] { "Automation", "FunctionalTest", "Material", "DerivedDataCache", "ShaderCompilers", "Texture" }, false);
+		}
+
 		/// <summary>
 		/// Return all entries for the specified channel. E.g. "OrionGame" will
 		/// return all entries starting with LogOrionGame
diff --git a/Engine/Source/Runtime/Engine/Private/Tests/AutomationCommon.cpp b/Engine/Source/Runtime/Engine/Private/Tests/AutomationCommon.cpp
--- a/Engine/Source/Runtime/Engine/Private/Tests/AutomationCommon.cpp
+++ b/Engine/Source/Runtime/Engine/Private/Tests/AutomationCommon.cpp
@@ -22,6 +22,8 @@
 #include "StereoRendering.h"
 #include "Misc/PackageName.h"
 #include "TextureCompiler.h"
+#include "Tests/AutomationTestSettings.h"
+#include "GameMapsSettings.h"
 
 #if WITH_AUTOMATION_TESTS
 
@@ -668,4 +670,135 @@ void RequestImageComparison(const FString& InImageName, int32 InWidth, int32 InH
 #endif
 }
 
+
+/**
+ * Write a string to editor automation tests log
+ */
+DEFINE_ENGINE_LATENT_AUTOMATION_COMMAND_ONE_PARAMETER(FEngineAutomationLogCommand, FString, LogText);
+
+bool FEngineAutomationLogCommand::Update()
+{
+	UE_LOG(LogEngineAutomationTests, Log, TEXT("%s"), *LogText);
+	return true;
+}
+
+/**
+ * Generic Pie Test for projects.
+ * By default this test will PIE the lit of MapsToPIETest from automation settings. if that is empty it will PIE the default editor and game (if they're different)
+ * maps.
+ *
+ * If the editor session was started with a map on the command line then that's the only map that will be PIE'd. This allows project to set up tests that PIE
+ * a list of maps from an external source.
+ */
+IMPLEMENT_SIMPLE_AUTOMATION_TEST(FProjectMapsCycleTest, "Project.Maps.Cycle", EAutomationTestFlags::ClientContext | EAutomationTestFlags::ProductFilter)
+
+/**
+ * Execute the loading of one map to verify PIE works
+ *
+ * @param Parameters - Unused for this test
+ * @return	TRUE if the test was successful, FALSE otherwise
+ */
+bool FProjectMapsCycleTest::RunTest(const FString& Parameters)
+{
+	UAutomationTestSettings const* AutomationTestSettings = GetDefault<UAutomationTestSettings>();
+	check(AutomationTestSettings);
+
+	// todo , move to automation settings
+	int CycleCount = 2;
+	TArray<FString> CycleMaps;
+
+	FString ParsedMapName;
+
+	if (FParse::Value(FCommandLine::Get(), TEXT("map="), ParsedMapName))
+	{
+		TArray<FString> MapList;
+
+		ParsedMapName.ParseIntoArray(MapList, TEXT("+"), true);
+
+		for (const FString& Map : MapList)
+		{
+			FString ActualName = Map;
+			// If the specified package exists
+			if (FPackageName::SearchForPackageOnDisk(Map, NULL, &ActualName) &&
+				// and it's a valid map file
+				FPaths::GetExtension(ActualName, /*bIncludeDot=*/true) == FPackageName::GetMapPackageExtension())
+			{
+				CycleMaps.Add(ActualName);
+				UE_LOG(LogEngineAutomationTests, Display, TEXT("Found Map %s on command line. Cycle Test will be use this map"), *ActualName);
+			}
+			else
+			{
+				UE_LOG(LogEngineAutomationTests, Fatal, TEXT("Cound not find package for Map '%s' specified on command line."), *ActualName);
+			}
+		}
+	}
+
+	FParse::Value(FCommandLine::Get(), TEXT("map.cycles="), CycleCount);
+
+	// If there was no command line map then default to the project settings
+	if (CycleMaps.Num() == 0)
+	{
+		// If the project has maps configured for PIE then use those
+	#if 0
+		if (AutomationTestSettings->MapsToPIETest.Num())
+		{
+			for (const FString& Map : AutomationTestSettings->MapsToPIETest)
+			{
+				CycleMaps.Add(Map);
+			}
+		}
+		else
+	#endif
+		{
+			
+			UGameMapsSettings const* MapSettings = GetDefault<UGameMapsSettings>();
+
+			if (MapSettings->GetGameDefaultMap().Len())
+			{
+				FString StartupMap = MapSettings->GetGameDefaultMap();
+				// Else pick the editor startup and game startup maps (if they are different).
+				UE_LOG(LogEngineAutomationTests, Display, TEXT("No MapsToCycle specified in DefaultEngine.ini [/Script/Engine.AutomationTestSettings]. Using GameStartup Map %s"), *StartupMap);
+				CycleMaps.Add(StartupMap);
+			}
+		}
+	}
+
+	// Uh-oh
+	if (CycleMaps.Num() == 0)
+	{
+		UE_LOG(LogEngineAutomationTests, Fatal, TEXT("No automation or default maps are configured for cycling!"));
+	}
+
+	for (int i = 1; i <= CycleCount; i++)
+	{
+		AddCommand(new FEngineAutomationLogCommand(FString::Printf(TEXT("Starting Project.Maps Cycle (%d/%d)"), i, CycleCount)));
+		for (const FString& Map : CycleMaps)
+		{
+			FString MapPackageName = FPackageName::ObjectPathToPackageName(Map);
+
+			if (!FPackageName::IsValidObjectPath(MapPackageName))
+			{
+				if (!FPackageName::SearchForPackageOnDisk(MapPackageName, NULL, &MapPackageName))
+				{
+					UE_LOG(LogEditorAutomationTests, Error, TEXT("Couldn't resolve map for PIE test from %s to valid package name!"), *MapPackageName);
+					continue;
+				}
+			}
+
+			AddCommand(new FEngineAutomationLogCommand(FString::Printf(TEXT("LoadMap-Begin: %s"), *MapPackageName)));
+			AddCommand(new FLoadGameMapCommand(Map));
+			AddCommand(new FEngineAutomationLogCommand(FString::Printf(TEXT("LoadMap-End: %s"), *MapPackageName)));
+			AddCommand(new FEngineAutomationLogCommand(FString::Printf(TEXT("MapWait-Begin: %s"), *MapPackageName)));
+			AddCommand(new FWaitForShadersToFinishCompilingInGame());
+			AddCommand(new FWaitForSpecifiedMapToLoadCommand(MapPackageName)); 
+			AddCommand(new FWaitLatentCommand(AutomationTestSettings->PIETestDuration));
+			AddCommand(new FEngineAutomationLogCommand(FString::Printf(TEXT("MapWait-End: %s"), *Map)));
+
+		}
+		AddCommand(new FEngineAutomationLogCommand(FString::Printf(TEXT("Ended Project.Maps Cycle (%d/%d)"), i, CycleCount)));
+	}
+
+	return true;
+}
+
 #endif //WITH_DEV_AUTOMATION_TESTS
