Added Mac options to Topaz tests

Added StarterContent tests

diff --git a/Engine/Build/Graph/Tasks/BuildAndTestProject.xml b/Engine/Build/Graph/Tasks/BuildAndTestProject.xml
--- a/Engine/Build/Graph/Tasks/BuildAndTestProject.xml
+++ b/Engine/Build/Graph/Tasks/BuildAndTestProject.xml
@@ -20,8 +20,8 @@
 	<Do If="$(WithBATDefaults)">
 		<Property Name="DefaultEditorPlatforms" Value="" />
 		<Property Name="DefaultTargetPlatforms" Value="" />
-		<Property Name="DefaultEditorTestList" Value="UE.EditorAutomation(RunTest=Project.Maps.PIE)" />
-		<Property Name="DefaultTargetTestList" Value="UE.TargetAutomation(RunTest=Project.Maps.Cycle)" />
+		<Property Name="DefaultEditorTestList" Value="" />
+		<Property Name="DefaultTargetTestList" Value="" />
 	</Do>
 
 	<!-- This will define options for all platforms and create lists called EditorPlatforms, RequiredEditorPlatforms, and TargetPlatforms. E.g. TargetPlatforms=Win64+Mac+PS4 -->
