Build scripts -

Added some missing platforms to BuildAndTestProject.xml
Added support for building different configs and MapIniSection
Updated ShooterGame, QAGame, and other build scripts

Exposed options to Horde UI

diff --git a/Engine/Build/Graph/Tasks/BuildAndTestProject.xml b/Engine/Build/Graph/Tasks/BuildAndTestProject.xml
--- a/Engine/Build/Graph/Tasks/BuildAndTestProject.xml
+++ b/Engine/Build/Graph/Tasks/BuildAndTestProject.xml
@@ -3,11 +3,21 @@
 	<EnvVar Name="UE_HORDE_JOBID"/>
 	
 	<!-- Project Options -->
-	<Option Name="ProjectName" DefaultValue="" Restrict="" Description="Name of Project"/>
+	<Option Name="ProjectName" DefaultValue="None" Restrict=".+" Description="Name of Project"/>
+	<Option Name="ProjectPath" DefaultValue="None" Restrict=".+" Description="Path to the folder that contains your project"/>
+
 	<Option Name="TargetName" DefaultValue="" Restrict="" Description="If your project supports both Game and Client you can set this, otherwise one will be selected"/>
 
-	<Option Name="ProjectPath" DefaultValue="" Restrict=".*" Description="Path to the folder that contains the project (should not include the .uproject file!)" />
-	<Option Name="PlatformsToSkipTests" DefaultValue="" Restrict=".*" Description="List of platforms that tests won't be run on" />
+
+	<Option Name="ExtraEditorCompileArguments" DefaultValue="" Restrict="" Description="Extra arguments to use when building the editor" />
+	<Option Name="ExtraTargetCompileArguments" DefaultValue="" Restrict="" Description="Extra arguments to use when building targets" />
+	<Option Name="MapIniSection" DefaultValue="" Restrict="" Description="MapIniSection argument to pass when cooking" />
+	<Option Name="ExtraCookArguments" DefaultValue="" Restrict="" Description="Extra arguments to use when cooking" />
+	<Option Name="ExtraStageAndPackageArguments" DefaultValue="" Restrict="" Description="Extra arguments to use when staging and packaging" />
+
+	<Option Name="PlatformsToSkipTests" DefaultValue="HoloLens+IOS+Lumin+Stadia+TVOS" Restrict=".*" Description="List of platforms that tests won't be run on" />
+
+	<Option Name="UseIncrementalAgents" DefaultValue="false" Restrict="true|false" Description="Use incremental agents for building" />
 
 	<Option Name="WithBATDefaults" DefaultValue="true" Description="Used by parent scripts to specify defaults when including this script"/>
 
@@ -20,6 +30,7 @@
 	<Do If="$(WithBATDefaults)">
 		<Property Name="DefaultEditorPlatforms" Value="" />
 		<Property Name="DefaultTargetPlatforms" Value="" />
+		<Property Name="DefaultTargetConfigurations" Value="Development" />
 		<Property Name="DefaultEditorTestList" Value="" />
 		<Property Name="DefaultTargetTestList" Value="" />
 	</Do>
@@ -58,18 +69,20 @@
 	<!-- At this point we have all options so calculate a few things -->
 
 	<!-- Validate that the project exists-->
-	<Error Message="ProjectName is not set" If="'$(ProjectName)' == ''" />
+	<Error Message="ProjectName and ProjectPath must be specified" If="'$(ProjectName)' == 'None' or '$(ProjectPath)' == 'None'" />
 	<Property Name="TargetProject" Value="$(ProjectPath)/$(ProjectName).uproject" />
 	<Error Message="Project file $(TargetProject) does not exist" If="!Exists('$(TargetProject)')" />
 
+
 	<!-- If the user didn't supply a TargetName then try to auto detect is. This is reliable uness the project has 'Game' in it but the target doesn't -->
 	<Do If="'$(TargetName)' == ''">
 		<!-- 
 			If the user didn't supply a TargetName, try to figure it out based on GameName.Target or GameNameClient.target.
-			Note - if the project is ShooterGame and the target is ShooterClient.txt then this will fail and we'll tell the user -->
+			Note - if the project is ShooterGame and the target is ShooterClient.txt then this will fail and we'll tell the user 
+		-->
 		<Property Name="TargetName" Value="$(ProjectName)Client" If="Exists('$(ProjectPath)/Source/$(ProjectName)Client.Target.cs')" />
 		<Property Name="TargetName" Value="$(ProjectName)" If="Exists('$(ProjectPath)/Source/$(ProjectName).Target.cs')" />
-
+		
 		<Do If="'$(TargetName)' == ''">
 			<Do If="Exists('$(ProjectPath)/Source/$(ProjectName)/$(ProjectName).Build.cs')">
 				<Error Message="The project has source but the name the target couldn't be determined. use -set:TargetName:Foo to ensure things build correctly" />
@@ -106,6 +119,8 @@
 	<ForEach Name="EditorPlatform" Values="$(RequiredEditorPlatforms)" Separator="+">
 		<!-- Set Agent Types -->
 		<Property Name="HostAgentType" Value="$(EditorPlatform)" />
+		<!-- Use incremental agents for building? -->
+		<Property Name="HostAgentType" Value="Incremental$(HostAgentType)" If="$(UseIncrementalAgents)" />
 		<Property Name="TestAgentType" Value="Test$(EditorPlatform)" />
 
 		<!-- Test nodes for this editor platform -->
@@ -121,7 +136,7 @@
 				<Compile Target="CrashReportClientEditor" Platform="$(EditorPlatform)" Configuration="Shipping"/>
 			</Node>
 			<Node Name="Compile $(ProjectName)Editor $(EditorPlatform)" Requires="Compile $(ProjectName)Tools $(EditorPlatform)">
-				<Compile Target="$(ProjectName)Editor" Platform="$(EditorPlatform)" Configuration="Development" />
+				<Compile Target="$(ProjectName)Editor" Platform="$(EditorPlatform)" Configuration="Development" Arguments="$(ExtraEditorCompileArguments)"/>
 			</Node>
 			<!--Add these nodes to our dependency list -->
 			<Property Name="GeneratedToolNodes" Value="Compile $(ProjectName)Editor $(EditorPlatform); Compile $(ProjectName)Tools $(EditorPlatform)" If="!$(SkipBuild)" />
@@ -177,24 +192,32 @@
 
 		<!-- Declare type of host required to build & cook this platform-->
 		<Property Name="HostAgentType" Value="Win64" />
-		<Property Name="HostAgentType" Value="Mac" If="'$(TargetPlatform)' == 'Mac' or '$(TargetPlatform)' == 'IOS'"/>
+		<Property Name="HostAgentType" Value="Mac" If="'$(TargetPlatform)' == 'Mac' or '$(TargetPlatform)' == 'IOS' or '$(TargetPlatform)' == 'tvOS'"/>
 		<!-- On consoles we can use regualr agents, but locally hosted tests need a Test agent-->
 		<Property Name="TestAgentType" Value="$(HostAgentType)" />
 		<Property Name="TestAgentType" Value="Test$(HostAgentType)" If="'$(TargetPlatform)' == 'Win64' or '$(TargetPlatform)' == 'Mac'"/>
 
+		<!-- Use incremental agents for building? -->
+		<Property Name="BuildAgentType" Value="$(HostAgentType)" />
+		<Property Name="BuildAgentType" Value="Incremental$(HostAgentType)" If="$(UseIncrementalAgents)" />
+
 		<!-- Building and cooking will require these nodes from the appropriate host. They should have been built above... -->
 		<Property Name="PlatformToolsNodeName" Value="Compile $(ProjectName)Tools $(HostAgentType)" />
 		<Property Name="PlatformEditorNodeName" Value="Compile $(ProjectName)Editor $(HostAgentType)" />
 		<!-- Define these names up front as we use some of them several times -->
 		<Property Name="PlatformAgentName" Value="$(DisplayTargetName) $(TargetPlatform)" />
-		<Property Name="PlatformCompileNodeName" Value="Compile $(TargetName) $(TargetPlatform) $(TargetConfiguration)" />
-		<Property Name="PlatformCookNodeName" Value="Cook $(DisplayTargetName) $(TargetPlatform) $(TargetConfiguration)" />
-		<Property Name="PlatformPackageNodeName" Value="Package $(DisplayTargetName) $(TargetPlatform) $(TargetConfiguration)" />
-		<Property Name="PlatformPublishNodeName" Value="Publish $(DisplayTargetName) $(TargetPlatform) $(TargetConfiguration)" />
-		<Property Name="PlatformBootTestNodeName" Value="$(DisplayTargetName) $(TargetPlatform) Test=BootTest" />
-
-		<!-- This will be a generated string of all tests for this platform that can be tied to an aggregate -->
-		<Property Name="PlatformTestNodeNames" Value="" />
+		<Property Name="PlatformCompileNodeName" Value="Compile $(TargetName) $(TargetPlatform)" />
+		<Property Name="PlatformCookNodeName" Value="Cook $(DisplayTargetName) $(TargetPlatform)" />
+		<Property Name="PlatformPackageNodeName" Value="Package $(DisplayTargetName) $(TargetPlatform)" />
+		<Property Name="PlatformPublishNodeName" Value="Publish $(DisplayTargetName) $(TargetPlatform)" />
+
+		<!-- These will be a generated string of all configurations and builds/cooks/packages/tests for this platform that can be tied to aggregates -->
+		<Property Name="AllPlatformBuildNodes" Value="" />
+		<Property Name="AllPlatformCookNodes" Value="" />
+		<Property Name="AllPlatformPackageNodes" Value="" />
+		<Property Name="AllPlatformPublishNodes" Value="" />
+		<Property Name="AllPlatformBootTestNodes" Value="" />
+		<Property Name="AllPlatformTestNodes" Value="" />
 		
 		<!-- Set CookPlatform since it isn't always the same as the target -->
 		<Property Name="CookPlatform" Value="$(TargetPlatform)" />
@@ -203,89 +226,103 @@
 		<Property Name="CookPlatform" Value="$(CookPlatform)Client" If="'$(TargetType)' == 'Client'" />
 
 		<!-- Build the client and if necessary a server-->
-		<Agent Name="Build $(PlatformAgentName)" Type="$(HostAgentType)">
-			<Node Name="$(PlatformCompileNodeName)" Requires="$(PlatformToolsNodeName)">
-				<Compile Target="$(TargetName)" Platform="$(TargetPlatform)" Configuration="$(TargetConfiguration)"/>
-				<Do If="'$(DedicatedServer)'" >
-					<Compile Target="$(TargetName)Server" Platform="Win64" Configuration="$(TargetConfiguration)"/>
-				</Do>
-			</Node>
+		<Agent Name="Build $(PlatformAgentName)" Type="$(BuildAgentType)">
+			<ForEach Name="TargetConfiguration" Values="$(TargetConfigurations)" Separator="+">
+				<Node Name="$(PlatformCompileNodeName) $(TargetConfiguration)" Requires="$(PlatformToolsNodeName)">
+					<Compile Target="$(TargetName)" Platform="$(TargetPlatform)" Configuration="$(TargetConfiguration)" Arguments="$(ExtraTargetCompileArguments)"/>
+					<Do If="'$(DedicatedServer)'" >
+						<Compile Target="$(TargetName)Server" Platform="Win64" Configuration="$(TargetConfiguration)" Arguments="$(ExtraTargetCompileArguments)"/>
+					</Do>				
+				</Node>
+				<Property Name="AllPlatformBuildNodes" Value="$(AllPlatformBuildNodes);$(PlatformCompileNodeName) $(TargetConfiguration)" />
+			</ForEach>
 		</Agent>
 
 		<!-- Cook for the platform. This can run at the same time as the above since we don't need the exe -->
 		<Agent Name="Cook $(PlatformAgentName)" Type="$(HostAgentType)">
 			<Node Name="$(PlatformCookNodeName)" Requires="$(PlatformEditorNodeName)">
-				<Cook Project="$(TargetProject)" Platform="$(CookPlatform)" Arguments="" />				
+				<Property Name="CookArgs" Value="$(CookPlatform)" />				
+				<Property Name="CookArgs" Value="-MapIniSection=$(MapIniSection) $(CookPlatform)" If="'$(MapIniSection)' != ''"/>				
+				<Cook Project="$(TargetProject)" Platform="$(CookPlatform)" Arguments="$(CookArgs)"/>				
 			</Node>			
+			<Property Name="AllPlatformCookNodes" Value="$(AllPlatformCookNodes);$(PlatformCookNodeName)" />
 		</Agent>
 
 		<!-- Now package the exe and cooked data-->
 		<Agent Name="Package $(PlatformAgentName)" Type="$(HostAgentType)">
-			<Node Name="$(PlatformPackageNodeName)" Requires="$(PlatformCompileNodeName);$(PlatformCookNodeName)">
-				<Property Name="BCRArgs" Value="-project=$(TargetProject) -platform=$(TargetPlatform) -configuration=$(TargetConfiguration) -skipbuild -skipcook -pak -stage -archive -archivedirectory=&quot;$(ProjectOutputDirectory)&quot;" />
-				<Property Name="BCRArgs" Value="$(BCRArgs) -client" If="'$(TargetType)' == 'Client'" />
-				<Property Name="BCRArgs" Value="$(BCRArgs) -server" If="'$(DedicatedServer)'" />
-				<Property Name="BCRArgs" Value="$(BCRArgs) -package" If="!ContainsItem('$(PlatformsToSkipPackaging)','$(TargetPlatform)','+')" />		
-				<Command Name="BuildCookRun" Arguments="$(BCRArgs)" />
-			</Node>
-			<Node Name="$(PlatformPublishNodeName)" Requires="$(PlatformPackageNodeName)">
-				<!-- Typically only build machines publish -->
-				<Do If="'$(ProjectPublishPath)' != ''" >
-					<Copy Files="..." From="$(ProjectOutputDirectory)/$(CookPlatform)" To="$(ProjectPublishPath)/$(CookPlatform)"/>
-				</Do>
-			</Node>
+			<ForEach Name="PackageConfiguration" Values="$(PackageConfigurations)" Separator="+">
+				<Node Name="$(PlatformPackageNodeName) $(PackageConfiguration)" Requires="$(PlatformCompileNodeName) $(PackageConfiguration);$(PlatformCookNodeName)">
+						<Property Name="BCRArgs" Value="-project=$(TargetProject) -platform=$(TargetPlatform) -configuration=$(PackageConfiguration) $(ExtraStageAndPackageArguments) -skipbuild -skipcook -pak -stage -archive -archivedirectory=&quot;$(ProjectOutputDirectory)/$(PackageConfiguration)&quot;" />
+						<Property Name="BCRArgs" Value="$(BCRArgs) -client" If="'$(TargetType)' == 'Client'" />
+						<Property Name="BCRArgs" Value="$(BCRArgs) -server" If="'$(DedicatedServer)'" />
+						<Property Name="BCRArgs" Value="$(BCRArgs) -package" If="!ContainsItem('$(PlatformsToSkipPackaging)','$(TargetPlatform)','+')" />		
+						<Command Name="BuildCookRun" Arguments="$(BCRArgs)" />
+				</Node>
+				<Node Name="$(PlatformPublishNodeName) $(PackageConfiguration)" Requires="$(PlatformPackageNodeName) $(PackageConfiguration)">
+					<!-- Typically only build machines publish -->
+					<Do If="'$(ProjectPublishPath)' != ''" >
+						<Copy Files="..." From="$(ProjectOutputDirectory)/$(PackageConfiguration)/$(CookPlatform)" To="$(ProjectPublishPath)/$(PackageConfiguration)/$(CookPlatform)"/>
+					</Do>
+				</Node>
+				<Property Name="AllPlatformCookNodes" Value="$(AllPlatformCookNodes);$(PlatformPackageNodeName) $(PackageConfiguration)" />
+				<Property Name="AllPlatformPublishNodes" Value="$(AllPlatformPublishNodes);$(PlatformPublishNodeName) $(PackageConfiguration)" />
+			</ForEach>
 		</Agent>			
 
 		<!-- Test Time! -->		
-		<Agent Name="Test $(PlatformAgentName)" Type="$(TestAgentType)">	
+		<Agent Name="Test $(PlatformAgentName)" Type="$(TestAgentType)">		
 			<!-- Prerequisites for the test to run -->
 			<Property Name="TestPrerequisites" Value="$(PlatformEditorNodeName)" />
-			<!-- Declare a boot test for the target-->	
-			<Node Name="$(PlatformBootTestNodeName)" After="$(PlatformPublishNodeName)">
+			<ForEach Name="TestConfiguration" Values="$(TestConfigurations)" Separator="+">
+				<!-- Declare a boot test for the target-->	
+				<Property Name="BootTestName" Value="$(DisplayTargetName) $(TargetPlatform) $(TestConfiguration) BootTest" />
+				<Node Name="$(BootTestName)" After="$(PlatformPublishNodeName) $(TestConfiguration)">
 				<Do If="'$(AutomationReportOutputDirectory)' != ''">
-					<Property Name="OutputFragmentPath" Value="$(TargetPlatform)/UE.BootTest" />
+					<Property Name="OutputFragmentPath" Value="$(TargetPlatform)/$(TestConfiguration)/UE.BootTest" />
 					<Property Name="OutputReportPath" Value="$(AutomationReportOutputDirectory)/$(OutputFragmentPath)" />
 					<Property Name="GauntletReportArgs" Value="-logdir=&quot;$(OutputReportPath)&quot; -artifactname=Gauntlet -ReportExportPath=&quot;$(OutputReportPath)&quot;" />
 					<Property Name="GauntletReportArgs" Value="$(GauntletReportArgs) -ReportURL=&quot;$(AutomationReportUrl)/$(OutputFragmentPath)/&quot;" If="'$(AutomationReportUrl)' != ''" />
 				</Do>
-				<Command Name="RunUnreal" Arguments="-project=$(TargetProject) -test=&quot;UE.BootTest&quot; -platform=$(TargetPlatform) -configuration=$(TargetConfiguration) $(GauntletTargetTestArgs) $(GauntletReportArgs)" />
-			</Node>
-			<!-- Link the boot test in to the prerequisites and generated nodes if its enabled-->
-			<Do If="$(TargetBootTest)">
-				<Property Name="TestPrerequisites" Value="$(TestPrerequisites);$(PlatformBootTestNodeName)" />
-			</Do>
-			<!-- Go through the list of tests -->	
-			<ForEach Name="TestName" Values="$(TargetTestList)" Separator="+">
-				<!-- Run and report on the test-->
-				<Do If="'$(AutomationReportOutputDirectory)' != ''">
-					<Property Name="OutputFragmentPath" Value="$(TargetPlatform)/$(TestName)" />
-					<Property Name="OutputReportPath" Value="$(AutomationReportOutputDirectory)/$(OutputFragmentPath)" />
-					<Property Name="GauntletReportArgs" Value="-logdir=&quot;$(OutputReportPath)&quot; -artifactname=Gauntlet -ReportExportPath=&quot;$(OutputReportPath)&quot;" />
-					<Property Name="GauntletReportArgs" Value="$(GauntletReportArgs) -ReportURL=&quot;$(AutomationReportUrl)/$(OutputFragmentPath)/&quot;" If="'$(AutomationReportUrl)' != ''" />
-				</Do>
-				<Property Name="TestNodeName" Value="$(DisplayTargetName) $(TargetPlatform) Test=$(TestName)" />
-				<!-- Run the test. No 'Requires' as we run on the same machine locally and use the publish folder otherwose -->
-				<Node Name="$(TestNodeName)" Requires="$(TestPrerequisites)">
-					<Command Name="RunUnreal" Arguments="-project=$(TargetProject) -test=&quot;$(TestName)&quot; -platform=$(TargetPlatform) -configuration=$(TargetConfiguration) $(GauntletTargetTestArgs) $(GauntletReportArgs)"/>
+				<Command Name="RunUnreal" Arguments="-project=$(TargetProject) -test=&quot;UE.BootTest&quot; -platform=$(TargetPlatform) -configuration=$(TestConfiguration) $(GauntletTargetTestArgs) $(GauntletReportArgs)" />
 				</Node>
-				<!-- Add it to our list -->
-				<Property Name="PlatformTestNodeNames" Value="$(PlatformTestNodeNames);$(TestNodeName)"/>
-			</ForEach>	
+				<!-- Link the boot test in to the prerequisites and generated nodes if its enabled-->
+				<Do If="$(TargetBootTest)">
+					<Property Name="TestPrerequisites" Value="$(TestPrerequisites);$(BootTestName)" />
+					<Property Name="AllPlatformBootTestNodes" Value="$(AllPlatformBootTestNodes);$(BootTestName)"/>
+				</Do>
+				<!-- Go through the list of tests -->	
+				<ForEach Name="TestName" Values="$(TargetTestList)" Separator="+">
+					<!-- Run and report on the test-->
+					<Do If="'$(AutomationReportOutputDirectory)' != ''">
+						<Property Name="OutputFragmentPath" Value="$(TargetPlatform)/$(TestConfiguration)/$(TestName)" />
+						<Property Name="OutputReportPath" Value="$(AutomationReportOutputDirectory)/$(OutputFragmentPath)" />
+						<Property Name="GauntletReportArgs" Value="-logdir=&quot;$(OutputReportPath)&quot; -artifactname=Gauntlet -ReportExportPath=&quot;$(OutputReportPath)&quot;" />
+						<Property Name="GauntletReportArgs" Value="$(GauntletReportArgs) -ReportURL=&quot;$(AutomationReportUrl)/$(OutputFragmentPath)/&quot;" If="'$(AutomationReportUrl)' != ''" />
+					</Do>
+					<Property Name="TestNodeName" Value="$(DisplayTargetName) $(TargetPlatform) $(TestConfiguration) $(TestName)" />
+					<!-- Run the test. No 'Requires' as we run on the same machine locally and use the publish folder otherwose -->
+					<Node Name="$(TestNodeName)" Requires="$(TestPrerequisites)">
+						<Command Name="RunUnreal" Arguments="-project=$(TargetProject) -test=&quot;$(TestName)&quot; -platform=$(TargetPlatform) -configuration=$(TestConfiguration) $(GauntletTargetTestArgs) $(GauntletReportArgs)"/>
+					</Node>
+					<!-- Add it to our list -->
+					<Property Name="AllPlatformTestNodes" Value="$(AllPlatformTestNodes);$(TestNodeName)"/>
+				</ForEach>	
+			</ForEach>
 		</Agent>
 
 		<!-- Declare labels for CIS -->
-		<Label Category="Platforms" Name="$(ProjectName) $(TargetPlatform)" Requires="$(PlatformCompileNodeName)" />
-		<Label Category="Cook" Name="Cook $(DisplayTargetName) $(TargetPlatform)" Requires="$(PlatformCookNodeName)" Exclude="$(PlatformEditorNodeName)" />
-		<Label Category="Package" Name="Pkg $(DisplayTargetName) $(TargetPlatform)" Requires="$(PlatformPackageNodeName);$(PlatformPublishNodeName)" Exclude="$(PlatformCompileNodeName);$(PlatformCookNodeName)"/>
-		<Label Category="Test" Name="BootTest $(DisplayTargetName) $(TargetPlatform)" Requires="$(PlatformBootTestNodeName)" Exclude="$(PlatformEditorNodeName)" />
-		<Label Category="Test" Name="Test $(DisplayTargetName) $(TargetPlatform)" Requires="$(PlatformTestNodeNames)" Exclude="$(PlatformEditorNodeName);$(PlatformBootTestNodeName)"/>
+		<Label Category="Platforms" Name="$(ProjectName) $(TargetPlatform)" Requires="$(AllPlatformBuildNodes)" />
+		<Label Category="Cook" Name="Cook $(DisplayTargetName) $(TargetPlatform)" Requires="$(AllPlatformCookNodes)" Exclude="$(PlatformEditorNodeName)" />
+		<Label Category="Package" Name="Pkg $(DisplayTargetName) $(TargetPlatform)" Requires="$(AllPlatformPackageNodes);$(AllPlatformPublishNodes)" Exclude="$(AllPlatformBuildNodes);$(AllPlatformCookNodes)"/>
+		<Label Category="Test" Name="BootTest $(DisplayTargetName) $(TargetPlatform)" Requires="$(AllPlatformBootTestNodes)" Exclude="$(PlatformEditorNodeName)" />
+		<Label Category="Test" Name="Test $(DisplayTargetName) $(TargetPlatform)" Requires="$(AllPlatformTestNodes)" Exclude="$(PlatformEditorNodeName);$(AllPlatformBootTestNodes)"/>
 
 		<!-- Tie all steps into our platform tasks uless excluded -->
 		<Property Name="PlatformTasks" Value="" />
-		<Property Name="PlatformTasks" Value="$(PlatformCompileNodeName)" If="!$(SkipBuild)" />
-		<Property Name="PlatformTasks" Value="$(PlatformTasks);$(PlatformCookNodeName)" If="!$(SkipCook) and !$(SkipPublish)" />		
-		<Property Name="PlatformTasks" Value="$(PlatformTasks);$(PlatformPackageNodeName);$(PlatformPublishNodeName)" If="!$(SkipPublish)" />		
-		<Property Name="PlatformTasks" Value="$(PlatformTasks);$(PlatformBootTestNodeName);$(PlatformTestNodeNames)" If="!ContainsItem('$(PlatformsToSkipTests)','$(TargetPlatform)','+')" />		
+		<Property Name="PlatformTasks" Value="$(AllPlatformBuildNodes)" If="!$(SkipBuild)" />
+		<Property Name="PlatformTasks" Value="$(PlatformTasks);$(AllPlatformCookNodes)" If="!$(SkipCook) and !$(SkipPublish)" />		
+		<Property Name="PlatformTasks" Value="$(PlatformTasks);$(AllPlatformPackageNodes);$(AllPlatformPublishNodes)" If="!$(SkipPublish)" />	
+		<Property Name="PlatformTasks" Value="$(PlatformTasks);$(AllPlatformBootTestNodes);$(AllPlatformTestNodes)" If="!ContainsItem('$(PlatformsToSkipTests)','$(TargetPlatform)','+')" />		
 
 		<Property Name="GeneratedTargetNodes" Value="$(GeneratedTargetNodes);$(PlatformTasks)"/>
 	</ForEach>
diff --git a/Engine/Build/Graph/Tasks/Inc/PlatformOptions.xml b/Engine/Build/Graph/Tasks/Inc/PlatformOptions.xml
--- a/Engine/Build/Graph/Tasks/Inc/PlatformOptions.xml
+++ b/Engine/Build/Graph/Tasks/Inc/PlatformOptions.xml
@@ -24,47 +24,62 @@
 		OutProperty: RequiredEditorPlatforms - Editor platforms to build to cook targets      
 	-->
 
-<Option Name="EditorPlatforms" DefaultValue="" Description="Alternative to 'WithFooEditor', a list of editors to build and test. Win64+Mac etc"/>
-<Option Name="TargetPlatforms" DefaultValue="" Description="Alternative to 'WithFoo', a list of Platforms to build and test. PS4+XSX etc"/>
-
-<Option Name="WithWin64Editor" DefaultValue="false" Description="Build and test Win64 Editor"/>
-<Option Name="WithMacEditor" DefaultValue="false" Description="Build and test Mac Editor"/>
+<Property Name="ShowExplicitFlags" Value="$(IsBuildMachine)" />
 
-<Option Name="WithAndroid" DefaultValue="false" Description="Build and test Android"/>
-<Option Name="WithIOS" DefaultValue="false" Description="Build and test IOS"/>
-<Option Name="WithLumin" DefaultValue="false" Description="Build and test Lumin"/>
-<Option Name="WithMac" DefaultValue="false" Description="Build and test Mac"/>
-<Option Name="WithPS4" DefaultValue="false" Description="Build and test PS4"/>
-<Option Name="WithPS5" DefaultValue="false" Description="Build and test PS5"/>
-<Option Name="WithStadia" DefaultValue="false" Description="Build and test Stadia"/>
-<Option Name="WithSwitch" DefaultValue="false" Description="Build and test Switch"/>
-<Option Name="WithWin64" DefaultValue="false" Description="Build and test Win64"/>
-<Option Name="WithWinGDK" DefaultValue="false" Description="Build and test WinGDK"/>
-<Option Name="WithXboxOneGDK" DefaultValue="false" Description="Build and test XboxOneGDK"/>
-<Option Name="WithXSX" DefaultValue="false" Description="Build and test XSX"/>
+<Property Name="AllEditorPlatforms" Value="Mac+Win64" />
+<Property Name="AllTargetPlatforms" Value="Android+HoloLens+iOS+Lumin+Mac+PS4+PS5+Stadia+Switch+tvOS+Win64+WinGDK+XboxOneGDK+XSX" />
 
-<Option Name="TargetConfiguration" DefaultValue="Development" Description="Configuration to build and test"/>
+<Option Name="EditorPlatforms" DefaultValue="" Description="Alternative to 'WithFooEditor', a list of editors to build and test. Win64+Mac etc"/>
+<Option Name="TargetPlatforms" DefaultValue="" Description="Alternative to 'WithFoo', a list of Platforms to build and test. PS4+XSX etc"/>
 
-<Do If="'$(EditorPlatforms)' == ''" >
-	<Property Name="EditorPlatforms" Value="" />
-	<Property Name="EditorPlatforms" Value="$(EditorPlatforms)Win64+" If="$(WithWin64Editor)" />
-	<Property Name="EditorPlatforms" Value="$(EditorPlatforms)Mac+" If="$(WithMacEditor)" />
-</Do>
+<!-- Support simple One+Two+Three style -->
+<Option Name="TargetConfigurations" DefaultValue="" Description="Configuration(s) to build"/>
+<Option Name="PackageConfigurations" DefaultValue="$(TargetConfigurations)" Description="Configuration to package"/>
+<Option Name="TestConfigurations" DefaultValue="$(PackageConfigurations)" Description="Configuration to test"/>
+
+<!-- Support explicit flags for horde but hide them otherwise as the option list is overwhelming..-->
+<Do If="$(ShowExplicitFlags)" >
+	<Option Name="WithWin64Editor" DefaultValue="false" Description="Build and test Win64 Editor"/>
+	<Option Name="WithMacEditor" DefaultValue="false" Description="Build and test Mac Editor"/>
+
+	<Option Name="WithAndroid" DefaultValue="false" Description="Build and test Android"/>
+	<Option Name="WithHoloLens" DefaultValue="false" Description="Build and test HoloLens"/>
+	<Option Name="WithIOS" DefaultValue="false" Description="Build and test IOS"/>
+	<Option Name="WithLumin" DefaultValue="false" Description="Build and test Lumin"/>
+	<Option Name="WithMac" DefaultValue="false" Description="Build and test Mac"/>
+	<Option Name="WithPS4" DefaultValue="false" Description="Build and test PS4"/>
+	<Option Name="WithPS5" DefaultValue="false" Description="Build and test PS5"/>
+	<Option Name="WithStadia" DefaultValue="false" Description="Build and test Stadia"/>
+	<Option Name="WithSwitch" DefaultValue="false" Description="Build and test Switch"/>
+	<Option Name="WithTvOS" DefaultValue="false" Description="Build and test TvOS"/>
+	<Option Name="WithWin64" DefaultValue="false" Description="Build and test Win64"/>
+	<Option Name="WithWinGDK" DefaultValue="false" Description="Build and test WinGDK"/>
+	<Option Name="WithXboxOneGDK" DefaultValue="false" Description="Build and test XboxOneGDK"/>
+	<Option Name="WithXSX" DefaultValue="false" Description="Build and test XSX"/>
+
+	<Do If="'$(EditorPlatforms)' == ''" >
+		<Property Name="EditorPlatforms" Value="" />
+		<Property Name="EditorPlatforms" Value="$(EditorPlatforms)Win64+" If="$(WithWin64Editor)" />
+		<Property Name="EditorPlatforms" Value="$(EditorPlatforms)Mac+" If="$(WithMacEditor)" />
+	</Do>
 
-<Do If="'$(TargetPlatforms)' == ''" >
-	<Property Name="TargetPlatforms" Value="" />
-	<Property Name="TargetPlatforms" Value="$(TargetPlatforms)Android+" If="$(WithAndroid)" />
-	<Property Name="TargetPlatforms" Value="$(TargetPlatforms)IOS+" If="$(WithIOS)" />
-	<Property Name="TargetPlatforms" Value="$(TargetPlatforms)Lumin+" If="$(WithLumin)" />
-	<Property Name="TargetPlatforms" Value="$(TargetPlatforms)Mac+" If="$(WithMac)" />
-	<Property Name="TargetPlatforms" Value="$(TargetPlatforms)PS4+" If="$(WithPS4)" />
-	<Property Name="TargetPlatforms" Value="$(TargetPlatforms)PS5+" If="$(WithPS5)" />
-	<Property Name="TargetPlatforms" Value="$(TargetPlatforms)Stadia+" If="$(WithStadia)" />
-	<Property Name="TargetPlatforms" Value="$(TargetPlatforms)Switch+" If="$(WithSwitch)" />
-	<Property Name="TargetPlatforms" Value="$(TargetPlatforms)Win64+" If="$(WithWin64)" />
-	<Property Name="TargetPlatforms" Value="$(TargetPlatforms)WinGDK+" If="$(WithWinGDK)" />
-	<Property Name="TargetPlatforms" Value="$(TargetPlatforms)XboxOneGDK+" If="$(WithXboxOneGDK)" />
-	<Property Name="TargetPlatforms" Value="$(TargetPlatforms)XSX+" If="$(WithXSX)" />
+	<Do If="'$(TargetPlatforms)' == ''" >
+		<Property Name="TargetPlatforms" Value="" />
+		<Property Name="TargetPlatforms" Value="$(TargetPlatforms)Android+" If="$(WithAndroid)" />
+		<Property Name="TargetPlatforms" Value="$(TargetPlatforms)HoloLens+" If="$(WithHoloLens)" />
+		<Property Name="TargetPlatforms" Value="$(TargetPlatforms)IOS+" If="$(WithIOS)" />
+		<Property Name="TargetPlatforms" Value="$(TargetPlatforms)Lumin+" If="$(WithLumin)" />
+		<Property Name="TargetPlatforms" Value="$(TargetPlatforms)Mac+" If="$(WithMac)" />
+		<Property Name="TargetPlatforms" Value="$(TargetPlatforms)PS4+" If="$(WithPS4)" />
+		<Property Name="TargetPlatforms" Value="$(TargetPlatforms)PS5+" If="$(WithPS5)" />
+		<Property Name="TargetPlatforms" Value="$(TargetPlatforms)Stadia+" If="$(WithStadia)" />
+		<Property Name="TargetPlatforms" Value="$(TargetPlatforms)Switch+" If="$(WithSwitch)" />
+		<Property Name="TargetPlatforms" Value="$(TargetPlatforms)TVOS+" If="$(WithTVOS)" />
+		<Property Name="TargetPlatforms" Value="$(TargetPlatforms)Win64+" If="$(WithWin64)" />
+		<Property Name="TargetPlatforms" Value="$(TargetPlatforms)WinGDK+" If="$(WithWinGDK)" />
+		<Property Name="TargetPlatforms" Value="$(TargetPlatforms)XboxOneGDK+" If="$(WithXboxOneGDK)" />
+		<Property Name="TargetPlatforms" Value="$(TargetPlatforms)XSX+" If="$(WithXSX)" />
+	</Do>
 </Do>
 
 <!-- If still no platforms, pick the defaults that the script should have specified-->
@@ -84,12 +99,24 @@
 	<Property Name="TargetPlatforms" Value="" />
 </Do>
 
+<!-- 'None' means don't do any platforms. We need to do this because an empty string means use defaults.. -->
+<Do If="'$(EditorPlatforms)' == 'All'" >
+	<Property Name="EditorPlatforms" Value="$(AllEditorPlatforms)" />
+</Do>
+<Do If="'$(TargetPlatforms)' == 'All'" >
+	<Property Name="TargetPlatforms" Value="$(AllTargetPlatforms)" />
+</Do>
+
 <Property Name="RequiredEditorPlatforms" Value="$(EditorPlatforms)" />
 
 <ForEach Name="TargetPlatform" Values="$(TargetPlatforms)" Separator="+">
+	<!-- Check the platform is valid -->
+	<Do If="!ContainsItem('$(AllTargetPlatforms)','$(TargetPlatform)','+')" >
+		<Error Message="Platform $(TargetPlatform) is not a recognized platform!" />
+	</Do>
 	<!-- Declare type of host required to build & cook this platform-->
 	<Property Name="HostAgentType" Value="Win64" />
-	<Property Name="HostAgentType" Value="Mac" If="'$(TargetPlatform)' == 'Mac' or '$(TargetPlatform)' == 'IOS'"/>
+	<Property Name="HostAgentType" Value="Mac" If="'$(TargetPlatform)' == 'Mac' or '$(TargetPlatform)' == 'IOS' or '$(TargetPlatform)' == 'tvOS'"/>
 
 	<Do If="!ContainsItem('$(RequiredEditorPlatforms)','$(HostAgentType)', '+')">
 		<Property Name="RequiredEditorPlatforms" Value="$(RequiredEditorPlatforms)+$(HostAgentType)" />
@@ -99,4 +126,61 @@
 <!-- PS5 packaging isn't supported by Gauntlet yet :( )-->
 <Property Name="PlatformsToSkipPackaging" Value="PS5" />
 
+
+<!-- Support explicit flags for horde but hide them otherwise as the option list is overwhelming..-->
+<Do If="$(ShowExplicitFlags)" >
+	<Option Name="WithDebugTargets" DefaultValue="false" Description="Build Debug configurations"/>
+	<Option Name="WithDevelopmentTargets" DefaultValue="false" Description="Build Development configurations"/>
+	<Option Name="WithTestTargets" DefaultValue="false" Description="Build Test configurations"/>
+	<Option Name="WithShippingTargets" DefaultValue="false" Description="Build Shipping configurations"/>
+
+	<Option Name="WithDebugPackages" DefaultValue="false" Description="Build Debug packages"/>
+	<Option Name="WithDevelopmentPackages" DefaultValue="false" Description="Build Development packages"/>
+	<Option Name="WithTestPackages" DefaultValue="false" Description="Build Test packages"/>
+	<Option Name="WithShippingPackages" DefaultValue="false" Description="Build Shipping packages"/>
+
+	<Option Name="WithDebugTests" DefaultValue="false" Description="Run tests on Debug builds"/>
+	<Option Name="WithDevelopmentTests" DefaultValue="false" Description="Run tests on Development builds"/>
+	<Option Name="WithTestTests" DefaultValue="false" Description="Run tests on Test builds"/>
+	<Option Name="WithShippingTests" DefaultValue="false" Description="Run tests on Shipping builds"/>
+
+	<Do If="'$(TargetConfigurations)' == ''" >	
+		<Property Name="TargetConfigurations" Value="" />
+		<Property Name="TargetConfigurations" Value="$(TargetConfigurations)Debug+"  If="$(WithDebugTargets)" />
+		<Property Name="TargetConfigurations" Value="$(TargetConfigurations)Development+"  If="$(WithDevelopmentTargets)" />
+		<Property Name="TargetConfigurations" Value="$(TargetConfigurations)Test+"  If="$(WithTestTargets)" />
+		<Property Name="TargetConfigurations" Value="$(TargetConfigurations)Shipping+"  If="$(WithShippingTargets)" />
+	</Do>
+
+	<Do If="'$(PackageConfigurations)' == ''" >	
+		<Property Name="PackageConfigurations" Value="" />
+		<Property Name="PackageConfigurations" Value="$(PackageConfigurations)Debug+"  If="$(WithDebugPackages)" />
+		<Property Name="PackageConfigurations" Value="$(PackageConfigurations)Development+"  If="$(WithDevelopmentPackages)" />
+		<Property Name="PackageConfigurations" Value="$(PackageConfigurations)Test+"  If="$(WithTestPackages)" />
+		<Property Name="PackageConfigurations" Value="$(PackageConfigurations)Shipping+"  If="$(WithShippingPackages)" />
+	</Do>
+
+	<Do If="'$(TestConfigurations)' == ''" >	
+		<Property Name="TestConfigurations" Value="" />
+		<Property Name="TestConfigurations" Value="$(TestConfigurations)Debug+"  If="$(WithDebugTests)" />
+		<Property Name="TestConfigurations" Value="$(TestConfigurations)Development+"  If="$(WithDevelopmentTests)" />
+		<Property Name="TestConfigurations" Value="$(TestConfigurations)Test+"  If="$(WithTestTests)" />
+		<Property Name="TestConfigurations" Value="$(TestConfigurations)Shipping+"  If="$(WithShippingTests)" />
+	</Do>
+</Do>
+
+<Do If="'$(TargetConfigurations)' == ''" >
+	<Property Name="TargetConfigurations" Value="$(DefaultTargetConfigurations)" />
+</Do>
+
+<Do If="'$(PackageConfigurations)' == ''" >	
+	<Property Name="PackageConfigurations" Value="$(TargetConfigurations)" />
+</Do>
+
+<Do If="'$(TestConfigurations)' == ''" >	
+	<Property Name="TestConfigurations" Value="$(PackageConfigurations)" />
+</Do>
+
+
+
 </BuildGraph>
\ No newline at end of file
