Fixed issue locating gauntlet source when building multiple packages into subfolders

diff --git a/Engine/Build/Graph/Tasks/BuildAndTestProject.xml b/Engine/Build/Graph/Tasks/BuildAndTestProject.xml
--- a/Engine/Build/Graph/Tasks/BuildAndTestProject.xml
+++ b/Engine/Build/Graph/Tasks/BuildAndTestProject.xml
@@ -283,7 +283,8 @@
 					<Property Name="GauntletReportArgs" Value="-logdir=&quot;$(OutputReportPath)&quot; -artifactname=Gauntlet -ReportExportPath=&quot;$(OutputReportPath)&quot;" />
 					<Property Name="GauntletReportArgs" Value="$(GauntletReportArgs) -ReportURL=&quot;$(AutomationReportUrl)/$(OutputFragmentPath)/&quot;" If="'$(AutomationReportUrl)' != ''" />
 				</Do>
-				<Command Name="RunUnreal" Arguments="-project=$(TargetProject) -test=&quot;UE.BootTest&quot; -platform=$(TargetPlatform) -configuration=$(TestConfiguration) $(GauntletTargetTestArgs) $(GauntletReportArgs)" />
+				<Property Name="GauntletBuildArg" Value="-build=&quot;$(GauntletBuildRoot)/$(TestConfiguration)&quot;" />
+				<Command Name="RunUnreal" Arguments="-project=$(TargetProject) -test=&quot;UE.BootTest&quot; -platform=$(TargetPlatform) -configuration=$(TestConfiguration) $(GauntletBuildArg) $(GauntletTargetTestArgs) $(GauntletReportArgs)" />
 				</Node>
 				<!-- Link the boot test in to the prerequisites and generated nodes if its enabled-->
 				<Do If="$(TargetBootTest)">
@@ -302,7 +303,7 @@
 					<Property Name="TestNodeName" Value="$(DisplayTargetName) $(TargetPlatform) $(TestConfiguration) $(TestName)" />
 					<!-- Run the test. No 'Requires' as we run on the same machine locally and use the publish folder otherwose -->
 					<Node Name="$(TestNodeName)" Requires="$(TestPrerequisites)">
-						<Command Name="RunUnreal" Arguments="-project=$(TargetProject) -test=&quot;$(TestName)&quot; -platform=$(TargetPlatform) -configuration=$(TestConfiguration) $(GauntletTargetTestArgs) $(GauntletReportArgs)"/>
+						<Command Name="RunUnreal" Arguments="-project=$(TargetProject) -test=&quot;$(TestName)&quot; -platform=$(TargetPlatform) -configuration=$(TestConfiguration) $(GauntletBuildArg) $(GauntletTargetTestArgs) $(GauntletReportArgs)"/>
 					</Node>
 					<!-- Add it to our list -->
 					<Property Name="AllPlatformTestNodes" Value="$(AllPlatformTestNodes);$(TestNodeName)"/>
diff --git a/Engine/Build/Graph/Tasks/Inc/GauntletSettings.xml b/Engine/Build/Graph/Tasks/Inc/GauntletSettings.xml
--- a/Engine/Build/Graph/Tasks/Inc/GauntletSettings.xml
+++ b/Engine/Build/Graph/Tasks/Inc/GauntletSettings.xml
@@ -23,15 +23,15 @@
 	<Option Name="GauntletTimeout" DefaultValue="3600" Description="Max running time of a test in Gauntlet"/>
 	
 	<!-- Create the -build argument for Gauntlet based on whether builds are local or on the network-->
-	<Property Name="GauntletBuildArgs" Value="-build=&quot;$(ProjectOutputDirectory)&quot;" />
-	<Property Name="GauntletBuildArgs" Value="-build=&quot;$(ProjectPublishPath)&quot;" If ="'$(ProjectPublishPath)' != ''" />
+	<Property Name="GauntletBuildRoot" Value="$(ProjectOutputDirectory)" />
+	<Property Name="GauntletBuildRoot" Value="$(ProjectPublishPath)" If ="'$(ProjectPublishPath)' != ''" />
 
 	<Property Name="GauntletCommonArgs" Value="-maxduration=$(GauntletTimeOut) -unattended -verbose" />
 
 	<Property Name="GauntletDeviceArgs" Value="" />
 	<Property Name="GauntletDeviceArgs" Value="-deviceurl=http://deviceservice.epicgames.net -reboot -cleardevices -sandbox=&quot;$(ProjectName)&quot;" If ="$(IsBuildMachine)" />
 
-	<Property Name="GauntletTargetTestArgs" Value="$(GauntletBuildArgs) $(GauntletCommonArgs) $(GauntletDeviceArgs)" />
+	<Property Name="GauntletTargetTestArgs" Value="$(GauntletCommonArgs) $(GauntletDeviceArgs)" />
 	<Property Name="GauntletEditorTestArgs" Value="-build=editor $(GauntletCommonArgs)" />
 
     <!-- this is fully generated dynamiclly when needed, it's simply here so it can be used without worrying if its set -->
