Added some safety checks to build script that will add missing package/build types based on specified tests/packages

diff --git a/Engine/Build/Graph/Tasks/Inc/PlatformOptions.xml b/Engine/Build/Graph/Tasks/Inc/PlatformOptions.xml
--- a/Engine/Build/Graph/Tasks/Inc/PlatformOptions.xml
+++ b/Engine/Build/Graph/Tasks/Inc/PlatformOptions.xml
@@ -24,7 +24,7 @@
 		OutProperty: RequiredEditorPlatforms - Editor platforms to build to cook targets      
 	-->
 
-<Property Name="ShowExplicitFlags" Value="$(IsBuildMachine)" />
+<Property Name="ShowExplicitFlags" Value="true" />
 
 <Property Name="AllEditorPlatforms" Value="Mac+Win64" />
 <Property Name="AllTargetPlatforms" Value="Android+HoloLens+iOS+Lumin+Mac+PS4+PS5+Stadia+Switch+tvOS+Win64+WinGDK+XboxOneGDK+XSX" />
@@ -181,6 +181,20 @@
 	<Property Name="TestConfigurations" Value="$(PackageConfigurations)" />
 </Do>
 
+<!-- add up any package configurations that are needed by test configurations.. -->
+<ForEach Name="TestConfiguration" Values="$(TestConfigurations)" Separator="+">
+	<Do If="!ContainsItem('$(PackageConfigurations)','$(TestConfiguration)', '+')">
+		<Trace Message="Adding package configuration for $(TestConfiguration) due to it being set for testing" />
+		<Property Name="PackageConfigurations" Value="$(PackageConfigurations)+$(TestConfiguration)" />
+	</Do>
+</ForEach>
 
+<!-- add up any target configurations that are needed by package configurations.. -->
+<ForEach Name="PackageConfiguration" Values="$(PackageConfigurations)" Separator="+">
+	<Do If="!ContainsItem('$(TargetConfigurations)','$(PackageConfiguration)', '+')">
+		<Trace Message="Adding target configuration for $(PackageConfiguration) due to it being set for testing" />
+		<Property Name="TargetConfigurations" Value="$(TargetConfigurations)+$(PackageConfiguration)" />
+	</Do>
+</ForEach>
 
 </BuildGraph>
\ No newline at end of file
