More BuildGraph tweaks

diff --git a/Engine/Build/Graph/Tasks/BuildAndTestProject.xml b/Engine/Build/Graph/Tasks/BuildAndTestProject.xml
--- a/Engine/Build/Graph/Tasks/BuildAndTestProject.xml
+++ b/Engine/Build/Graph/Tasks/BuildAndTestProject.xml
@@ -8,7 +8,7 @@
 
 	<Option Name="TargetName" DefaultValue="" Restrict="" Description="If your project supports both Game and Client you can set this, otherwise one will be selected"/>
 
-
+	<Option Name="ExtraToolCompileArguments" DefaultValue="" Restrict="" Description="Extra arguments to use when building the tools" />
 	<Option Name="ExtraEditorCompileArguments" DefaultValue="" Restrict="" Description="Extra arguments to use when building the editor" />
 	<Option Name="ExtraTargetCompileArguments" DefaultValue="" Restrict="" Description="Extra arguments to use when building targets" />
 	<Option Name="MapIniSection" DefaultValue="" Restrict="" Description="MapIniSection argument to pass when cooking" />
@@ -101,6 +101,12 @@
 		<Property Name="PlatformsToSkipPackaging" Value="$(EditorPlatforms)+$(TargetPlatforms)" />
 	</Do>
 
+	<Do If="$(UseIncrementalAgents)" >
+		<Property Name="ExtraToolCompileArguments" Value="-allmodules -disableunity -nodebuginfo -2019 -Strict -WarningsAsErrors"/>
+		<Property Name="ExtraEditorCompileArguments" Value="-allmodules -disableunity -nodebuginfo -2019 -Strict -WarningsAsErrors"/>
+		<Property Name="ExtraTargetCompileArguments" Value="-nodebuginfo -WarningsAsErrors"/>
+	</Do>
+
 	<!-- Detect the type of the client based on the target name -->
 	<Property Name="TargetType" Value="Game" />
 	<Property Name="TargetType" Value="Client" If="Contains('$(TargetName)', 'Client')" />
