Tidied up some options and handle some user-selection-errors better in build script

diff --git a/Engine/Build/Graph/Tasks/BuildAndTestProject.xml b/Engine/Build/Graph/Tasks/BuildAndTestProject.xml
--- a/Engine/Build/Graph/Tasks/BuildAndTestProject.xml
+++ b/Engine/Build/Graph/Tasks/BuildAndTestProject.xml
@@ -279,18 +279,21 @@
 		<Agent Name="Test $(PlatformAgentName)" Type="$(TestAgentType)">		
 			<!-- Prerequisites for the test to run -->
 			<Property Name="TestPrerequisites" Value="$(PlatformEditorNodeName)" />
+			<!-- Run through all configs -->
 			<ForEach Name="TestConfiguration" Values="$(TestConfigurations)" Separator="+">
+				<!-- Location of the build for this configuration -->
+				<Property Name="GauntletBuildArg" Value="-build=&quot;$(GauntletBuildRoot)/$(TestConfiguration)&quot;" />
 				<!-- Declare a boot test for the target-->	
 				<Property Name="BootTestName" Value="$(DisplayTargetName) $(TargetPlatform) $(TestConfiguration) BootTest" />
 				<Node Name="$(BootTestName)" After="$(PlatformPublishNodeName) $(TestConfiguration)">
-				<Do If="'$(AutomationReportOutputDirectory)' != ''">
-					<Property Name="OutputFragmentPath" Value="$(TargetPlatform)/$(TestConfiguration)/UE.BootTest" />
-					<Property Name="OutputReportPath" Value="$(AutomationReportOutputDirectory)/$(OutputFragmentPath)" />
-					<Property Name="GauntletReportArgs" Value="-logdir=&quot;$(OutputReportPath)&quot; -artifactname=Gauntlet -ReportExportPath=&quot;$(OutputReportPath)&quot;" />
-					<Property Name="GauntletReportArgs" Value="$(GauntletReportArgs) -ReportURL=&quot;$(AutomationReportUrl)/$(OutputFragmentPath)/&quot;" If="'$(AutomationReportUrl)' != ''" />
-				</Do>
-				<Property Name="GauntletBuildArg" Value="-build=&quot;$(GauntletBuildRoot)/$(TestConfiguration)&quot;" />
-				<Command Name="RunUnreal" Arguments="-project=$(TargetProject) -test=&quot;UE.BootTest&quot; -platform=$(TargetPlatform) -configuration=$(TestConfiguration) $(GauntletBuildArg) $(GauntletTargetTestArgs) $(GauntletReportArgs)" />
+					<Do If="'$(AutomationReportOutputDirectory)' != ''">
+						<Property Name="OutputFragmentPath" Value="$(TargetPlatform)/$(TestConfiguration)/UE.BootTest" />
+						<Property Name="OutputReportPath" Value="$(AutomationReportOutputDirectory)/$(OutputFragmentPath)" />
+						<Property Name="GauntletReportArgs" Value="-logdir=&quot;$(OutputReportPath)&quot; -artifactname=Gauntlet -ReportExportPath=&quot;$(OutputReportPath)&quot;" />
+						<Property Name="GauntletReportArgs" Value="$(GauntletReportArgs) -ReportURL=&quot;$(AutomationReportUrl)/$(OutputFragmentPath)/&quot;" If="'$(AutomationReportUrl)' != ''" />
+					</Do>
+					<!-- Run the boot test -->
+					<Command Name="RunUnreal" Arguments="-project=$(TargetProject) -test=&quot;UE.BootTest&quot; -platform=$(TargetPlatform) -configuration=$(TestConfiguration) $(GauntletBuildArg) $(GauntletTargetTestArgs) $(GauntletReportArgs)" />
 				</Node>
 				<!-- Link the boot test in to the prerequisites and generated nodes if its enabled-->
 				<Do If="$(TargetBootTest)">
diff --git a/Engine/Build/Graph/Tasks/Inc/PlatformOptions.xml b/Engine/Build/Graph/Tasks/Inc/PlatformOptions.xml
--- a/Engine/Build/Graph/Tasks/Inc/PlatformOptions.xml
+++ b/Engine/Build/Graph/Tasks/Inc/PlatformOptions.xml
@@ -184,7 +184,7 @@
 <!-- add up any package configurations that are needed by test configurations.. -->
 <ForEach Name="TestConfiguration" Values="$(TestConfigurations)" Separator="+">
 	<Do If="!ContainsItem('$(PackageConfigurations)','$(TestConfiguration)', '+')">
-		<Trace Message="Adding package configuration for $(TestConfiguration) due to it being set for testing" />
+		<Trace Message="Adding package configuration for $(TestConfiguration) due to the config being set for testing" />
 		<Property Name="PackageConfigurations" Value="$(PackageConfigurations)+$(TestConfiguration)" />
 	</Do>
 </ForEach>
@@ -192,7 +192,7 @@
 <!-- add up any target configurations that are needed by package configurations.. -->
 <ForEach Name="PackageConfiguration" Values="$(PackageConfigurations)" Separator="+">
 	<Do If="!ContainsItem('$(TargetConfigurations)','$(PackageConfiguration)', '+')">
-		<Trace Message="Adding target configuration for $(PackageConfiguration) due to it being set for testing" />
+		<Trace Message="Adding target configuration for $(PackageConfiguration) due to the config being set for packaging" />
 		<Property Name="TargetConfigurations" Value="$(TargetConfigurations)+$(PackageConfiguration)" />
 	</Do>
 </ForEach>
