Take a different approach to setting the build location for tests due to some builder weirdness

diff --git a/Engine/Build/Graph/Tasks/BuildAndTestProject.xml b/Engine/Build/Graph/Tasks/BuildAndTestProject.xml
--- a/Engine/Build/Graph/Tasks/BuildAndTestProject.xml
+++ b/Engine/Build/Graph/Tasks/BuildAndTestProject.xml
@@ -281,8 +281,9 @@
 			<Property Name="TestPrerequisites" Value="$(PlatformEditorNodeName)" />
 			<!-- Run through all configs -->
 			<ForEach Name="TestConfiguration" Values="$(TestConfigurations)" Separator="+">
-				<!-- Location of the build for this configuration -->
-				<Property Name="GauntletBuildArg" Value="-build=&quot;$(GauntletBuildRoot)/$(TestConfiguration)&quot;" />
+				<!-- Create the -build argument for Gauntlet based on whether builds are local or on the network-->
+				<Property Name="GauntletBuildArg" Value="-build=&quot;$(ProjectOutputDirectory)/$(TestConfiguration)&quot;" />
+				<Property Name="GauntletBuildArg" Value="-build=&quot;$(ProjectPublishPath)/$(TestConfiguration)&quot;" If ="'$(ProjectPublishPath)' != ''" />
 				<!-- Declare a boot test for the target-->	
 				<Property Name="BootTestName" Value="$(DisplayTargetName) $(TargetPlatform) $(TestConfiguration) BootTest" />
 				<Node Name="$(BootTestName)" After="$(PlatformPublishNodeName) $(TestConfiguration)">
diff --git a/Engine/Build/Graph/Tasks/Inc/GauntletSettings.xml b/Engine/Build/Graph/Tasks/Inc/GauntletSettings.xml
--- a/Engine/Build/Graph/Tasks/Inc/GauntletSettings.xml
+++ b/Engine/Build/Graph/Tasks/Inc/GauntletSettings.xml
@@ -22,10 +22,6 @@
 	<!-- How long to give Gauntlet before a test is considered timed out (some tests will time out earlier if they detect inactivity)-->
 	<Option Name="GauntletTimeout" DefaultValue="3600" Description="Max running time of a test in Gauntlet"/>
 	
-	<!-- Create the -build argument for Gauntlet based on whether builds are local or on the network-->
-	<Property Name="GauntletBuildRoot" Value="$(ProjectOutputDirectory)" />
-	<Property Name="GauntletBuildRoot" Value="$(ProjectPublishPath)" If ="'$(ProjectPublishPath)' != ''" />
-
 	<Property Name="GauntletCommonArgs" Value="-maxduration=$(GauntletTimeOut) -unattended -verbose" />
 
 	<Property Name="GauntletDeviceArgs" Value="" />
