Fixed case in BuildAndTest graph that could result in no code compilations for a target if a script didn't supply all defaults

diff --git a/Engine/Build/Graph/Tasks/Inc/PlatformOptions.xml b/Engine/Build/Graph/Tasks/Inc/PlatformOptions.xml
--- a/Engine/Build/Graph/Tasks/Inc/PlatformOptions.xml
+++ b/Engine/Build/Graph/Tasks/Inc/PlatformOptions.xml
@@ -171,6 +171,7 @@
 
 <Do If="'$(TargetConfigurations)' == ''" >
 	<Property Name="TargetConfigurations" Value="$(DefaultTargetConfigurations)" />
+	<Property Name="TargetConfigurations" Value="Development" If="'$(TargetConfigurations)' == ''" />
 </Do>
 
 <Do If="'$(PackageConfigurations)' == ''" >	
