WorldPartitionHLODsBuilder
* Added options to generate/use a build manifest
* -BuildManifest=Filename.ini
* When calling with -SetupHLODs, provide "-BuildManifest=Filename.ini -BuilderCount=N" to have N workloads in the build manifest
* When calling with -BuildHLODs, provide "-BuildManifest=Filename.ini -BuilderIdx=i" to build workload i

Added UWorldPartitionBuilder::PreWorldInitialization()

Added support for distributed World Partition HLOD generation
* Runs multiple steps on different machines
* Setup: Done on a single machine, will create/update/delete HLOD actors as required. The distributed setup step takes a "BuilderCount" argument, which will be used to write a build manifest that specify the work that needs to be done by the N builders
* Build: Performed on multiple machines, read work to be done from the build manifest
* Submit: Final step, perform on a single machine. Gather files modified in all the previous steps and submit them in a single changelist

Intermediary files that are created/modified/deleted are transfered between the jobs as build products, and submitted in the last step.

To run with UAT:
RunUAT.bat BuildGraph -Script=Engine\Build\Graph\Examples\BuildWorldPartitionHLODs.xml -Target="HLOD Generation" -set:ProjectName=YourProjectName -set:EditorTarget=EditorTargetName -set:MapName=YourMapName -set:BuilderCount=N

To run steps individually:
* Setup (N = Builder count to use)
    YourProject YourMapName -run=WorldPartitionBuilderCommandlet -SCCProvider=Perforce -AllowCommandletRendering -Builder=WorldPartitionHLODsBuilder -SetupHLODs -DistributedBuild -BuilderCount=N
* Build (i from 0 to N-1):
    YourProject YourMapName -run=WorldPartitionBuilderCommandlet -SCCProvider=Perforce -AllowCommandletRendering -Builder=WorldPartitionHLODsBuilder -DistributedBuild -BuildHLODs -BuilderIdx=i
* Submit:
    YourProject YourMapName -run=WorldPartitionBuilderCommandlet -SCCProvider=Perforce -AllowCommandletRendering -Builder=WorldPartitionHLODsBuilder -DistributedBuild -SubmitHLODs

Changes integrated from UE5 main:

diff --git a/Engine/Build/Graph/Examples/BuildWorldPartitionHLODs.xml b/Engine/Build/Graph/Examples/BuildWorldPartitionHLODs.xml
new file mode 100644
--- /dev/null
+++ b/Engine/Build/Graph/Examples/BuildWorldPartitionHLODs.xml
@@ -0,0 +1,70 @@
+<?xml version='1.0' ?>
+<BuildGraph xmlns="http://www.epicgames.com/BuildGraph" xmlns:xsi="http://www.w3.org/2001/XMLSchema-instance" xsi:schemaLocation="http://www.epicgames.com/BuildGraph ../Schema.xsd" >
+
+	<Option Name="MapName" DefaultValue="" Description="The map for which HLODs will be generated."/>
+	<Option Name="ProjectName" DefaultValue="" Description="Name of the project."/>
+	<Option Name="EditorTarget" DefaultValue="" Description="Name of the editor target to use."/>
+	<Option Name="BuilderCount" DefaultValue="10" Description="Number of machines to use"/>
+
+	<!-- Perforce environment variables -->
+	<EnvVar Name="uebp_PORT"/>
+	<EnvVar Name="uebp_USER"/>
+	<EnvVar Name="uebp_CLIENT"/>
+		
+	<Property Name="HLODCommonBuilderArgs" Value="$(MapName) -Builder=WorldPartitionHLODsBuilder -DistributedBuild"/>
+	<Property Name="HLODCommonBuilderArgs" Value="-unattended -buildmachine -SCCProvider=Perforce -P4Port=$(uebp_PORT) -P4User=$(uebp_USER) -P4Client=$(uebp_CLIENT) $(HLODCommonBuilderArgs)" If="$(IsBuildMachine)"/>
+	
+	<Agent Name="HLOD Generation Prerequisites" Type="Win64">
+		<Node Name="Update Version Files">
+			<SetVersion Change="$(Change)" Branch="$(EscapedBranch)" If="$(IsBuildMachine)"/>
+		</Node>
+		
+		<Node Name="Compile Tools Win64" Requires="Update Version Files">
+			<Compile Target="UnrealHeaderTool" Platform="Win64" Configuration="Development"/>
+			<Compile Target="ShaderCompileWorker" Platform="Win64" Configuration="Development"/>
+		</Node>
+		
+		<Node Name="Compile $(EditorTarget) Win64">
+			<Compile Target="$(EditorTarget)" Platform="Win64" Configuration="Development"/>
+		</Node>
+	</Agent>
+	<Property Name="HLODSetupDependencies" Value="Compile $(EditorTarget) Win64;Compile Tools Win64"/>
+	
+	<Agent Name="HLOD Setup" Type="Win64">
+		<Node Name="HLOD Setup" Requires="$(HLODSetupDependencies)">
+			<Commandlet Name="WorldPartitionBuilderCommandlet" Project="$(ProjectName)" Arguments="$(HLODCommonBuilderArgs) -SetupHLODs -BuilderCount=$(BuilderCount)"/>
+			<GatherBuildProductsFromFile BuildProductsFile="$(RootDir)/HLODTemp/BuildProducts.txt"/>
+		</Node>
+	</Agent>
+
+	<Property Name="HLODBuildNodes" Value=""/>
+	<Property Name="HLODBuildCompleteNodes" Value=""/>
+
+	<ForEach Name="HLODBuildJobIndex" Values="0;1;2;3;4;5;6;7;8;9;10;11;12;13;14;15;16;17;18;19;20;21;22;23;24;25;26;27;28;29;30;31;32;33;34;35;36;37;38;39;40;41;42;43;44;45;46;47;48;49">
+		<Do If="'$(HLODBuildJobIndex)' &lt; '$(BuilderCount)'">
+			<Property Name="HLODBuildAgent" Value="HLOD Build Agent $(HLODBuildJobIndex)"/>
+			<Property Name="HLODBuildNode" Value="HLOD Build Job $(HLODBuildJobIndex)"/>
+			<Property Name="HLODBuildNodes" Value="$(HLODBuildNodes);$(HLODBuildNode)"/>
+			<Property Name="HLODBuildCompleteNode" Value="$(HLODBuildNode) Complete"/>
+			<Property Name="HLODBuildCompleteNodes" Value="$(HLODBuildCompleteNodes);$(HLODBuildCompleteNode)"/>
+			
+			<Agent Name="$(HLODBuildAgent)" Type="Win64">
+				<Node Name="$(HLODBuildNode)" Requires="HLOD Setup">
+					<Commandlet Name="WorldPartitionBuilderCommandlet" Project="$(ProjectName)" Arguments="$(HLODCommonBuilderArgs) -AllowCommandletRendering -BuildHLODs -BuilderIdx=$(HLODBuildJobIndex)"/>
+				</Node>
+				<Node Name="$(HLODBuildCompleteNode)" After="$(HLODBuildNode)">
+					<GatherBuildProductsFromFile BuildProductsFile="$(RootDir)/HLODTemp/BuildProducts.txt"/>
+				</Node>
+			</Agent>
+		</Do>
+	</ForEach>
+	
+	<Agent Name="HLOD Submit" Type="Win64_NonAWS">
+		<Node Name="HLOD Submit" Requires="$(HLODBuildNodes);$(HLODBuildCompleteNodes);HLOD Setup">
+			<Commandlet Name="WorldPartitionBuilderCommandlet" Project="$(ProjectName)" Arguments="$(HLODCommonBuilderArgs) -SubmitHLODs"/>
+		</Node>
+	</Agent>
+
+	<Aggregate Name="HLOD Generation" Requires="HLOD Submit"/>
+	
+</BuildGraph>
\ No newline at end of file
