Added option to BuildAndTestProject script to publish staged content
Cleaned up some options around paths to make it easier for scripts to customize them.
Fixed issue with GDKBuilds in Gauntlet where the build source was ignoring the specified recursion depth when searching

diff --git a/Engine/Build/Graph/Tasks/BuildAndTestProject.xml b/Engine/Build/Graph/Tasks/BuildAndTestProject.xml
--- a/Engine/Build/Graph/Tasks/BuildAndTestProject.xml
+++ b/Engine/Build/Graph/Tasks/BuildAndTestProject.xml
@@ -22,9 +22,10 @@
 	<Option Name="WithBATDefaults" DefaultValue="true" Description="Used by parent scripts to specify defaults when including this script"/>
 
 	<!-- 
-		Allow tests to specify their own defaults by specifying a WithBATDefaults property that is true. The reason 
-		we use an option is we need something to check, and a property in a parent script will always override 
-		an option. Note - this means if a parent script includes this script multiple tiles then it should
+		Allow project graphs that Include this file to specify their own defaults. This is done by specifying WithBATDefaults=false
+		then declaring *all* these default values (leaving them empty if no default is required). 
+		The reason we need to wrap these in WithBATDefaults is that the properties always neeed to be declared.
+		Note - this means if a parent script includes this script multiple tiles then it should
 		always WithBATDefaults true/false before doing so, and if false it should supply all these arguments.
 	-->
 	<Do If="$(WithBATDefaults)">
@@ -33,6 +34,9 @@
 		<Property Name="DefaultTargetConfigurations" Value="Development" />
 		<Property Name="DefaultEditorTestList" Value="" />
 		<Property Name="DefaultTargetTestList" Value="" />
+		<Property Name="NetworkTempRootOverride" Value="" />
+		<Property Name="NetworkPublishRootOverride" Value="" />
+		<Property Name="NetworkReportRootOverride" Value="" />
 	</Do>
 
 	<!-- This will define options for all platforms and create lists called EditorPlatforms, RequiredEditorPlatforms, and TargetPlatforms. E.g. TargetPlatforms=Win64+Mac+PS4 -->
@@ -257,17 +261,31 @@
 		<!-- Now package the exe and cooked data-->
 		<Agent Name="Package $(PlatformAgentName)" Type="$(HostAgentType)">
 			<ForEach Name="PackageConfiguration" Values="$(PackageConfigurations)" Separator="+">
+				<Property Name="PackageDir" Value="$(ProjectOutputDirectory)/$(PackageConfiguration)/Packages" />
+				<Property Name="StagedDir" Value="$(ProjectOutputDirectory)/$(PackageConfiguration)/Staged" />
+
+				<Property Name="IsPackaged" Value="true" />
+				<Property Name="IsPackaged" Value="false" If="ContainsItem('$(PlatformsToSkipPackaging)','$(TargetPlatform)','+')" />	
+
 				<Node Name="$(PlatformPackageNodeName) $(PackageConfiguration)" Requires="$(PlatformCompileNodeName) $(PackageConfiguration);$(PlatformCookNodeName)">
-						<Property Name="BCRArgs" Value="-project=$(TargetProject) -platform=$(TargetPlatform) -configuration=$(PackageConfiguration) $(ExtraStageAndPackageArguments) -skipbuild -skipcook -pak -stage -archive -archivedirectory=&quot;$(ProjectOutputDirectory)/$(PackageConfiguration)&quot;" />
+						<Property Name="BCRArgs" Value="-project=$(TargetProject) -platform=$(TargetPlatform) -configuration=$(PackageConfiguration) $(ExtraStageAndPackageArguments) -skipbuild -skipcook -pak -stage -StagingDirectory=&quot;$(StagedDir)&quot;" />
 						<Property Name="BCRArgs" Value="$(BCRArgs) -client" If="'$(TargetType)' == 'Client'" />
 						<Property Name="BCRArgs" Value="$(BCRArgs) -server" If="'$(DedicatedServer)'" />
-						<Property Name="BCRArgs" Value="$(BCRArgs) -package" If="!ContainsItem('$(PlatformsToSkipPackaging)','$(TargetPlatform)','+')" />		
+						<Property Name="BCRArgs" Value="$(BCRArgs) -package -archive -archivedirectory=&quot;$(PackageDir)&quot;" If="$(IsPackaged)" />		
 						<Command Name="BuildCookRun" Arguments="$(BCRArgs)" />
 				</Node>
 				<Node Name="$(PlatformPublishNodeName) $(PackageConfiguration)" Requires="$(PlatformPackageNodeName) $(PackageConfiguration)">
-					<!-- Typically only build machines publish -->
-					<Do If="'$(ProjectPublishPath)' != ''" >
-						<Copy Files="..." From="$(ProjectOutputDirectory)/$(PackageConfiguration)/$(CookPlatform)" To="$(ProjectPublishPath)/$(PackageConfiguration)/$(CookPlatform)"/>
+					<!-- 	
+						If this path is set we need to write to the network. This path will reflect set to scratch or semi-permanent storage based on whether
+						PublishPackages=true or we're a buld machine ( build machines always need to put builds on the network for testing)
+					-->
+					<Do If="'$(NetworkOutputDirectory)' != ''">						
+						<Do If="$(IsPackaged)" >
+							<Copy Files="..." From="$(PackageDir)/$(CookPlatform)" To="$(NetworkOutputDirectory)/$(CookPlatform)/$(PackageConfiguration)"/>
+						</Do>
+						<Do If="$(PublishStagedContent) or !$(IsPackaged)" >
+							<Copy Files="..." From="$(StagedDir)/$(CookPlatform)" To="$(NetworkOutputDirectory)/$(CookPlatform)/Staged"/>
+						</Do>
 					</Do>
 				</Node>
 				<Property Name="AllPlatformCookNodes" Value="$(AllPlatformCookNodes);$(PlatformPackageNodeName) $(PackageConfiguration)" />
@@ -283,7 +301,7 @@
 			<ForEach Name="TestConfiguration" Values="$(TestConfigurations)" Separator="+">
 				<!-- Create the -build argument for Gauntlet based on whether builds are local or on the network-->
 				<Property Name="GauntletBuildArg" Value="-build=&quot;$(ProjectOutputDirectory)/$(TestConfiguration)&quot;" />
-				<Property Name="GauntletBuildArg" Value="-build=&quot;$(ProjectPublishPath)/$(TestConfiguration)&quot;" If ="'$(ProjectPublishPath)' != ''" />
+				<Property Name="GauntletBuildArg" Value="-build=&quot;$(NetworkOutputDirectory)&quot;" If ="'$(NetworkOutputDirectory)' != ''" />
 				<!-- Declare a boot test for the target-->	
 				<Property Name="BootTestName" Value="$(DisplayTargetName) $(TargetPlatform) $(TestConfiguration) BootTest" />
 				<Node Name="$(BootTestName)" After="$(PlatformPublishNodeName) $(TestConfiguration)">
diff --git a/Engine/Build/Graph/Tasks/Inc/GauntletSettings.xml b/Engine/Build/Graph/Tasks/Inc/GauntletSettings.xml
--- a/Engine/Build/Graph/Tasks/Inc/GauntletSettings.xml
+++ b/Engine/Build/Graph/Tasks/Inc/GauntletSettings.xml
@@ -11,8 +11,6 @@
 		InProperty: ProjectPath - Path where the .uproject is located
         InProperty: ProjectName - Name of the project (e.g. <ProjectName>.uproject)
 		InProperty: EscapedBranch
-		InProperty: ProjectOutputDirectory - Local pathwhere builds are found
-		InProperty: ProjectPublishPath - Path (if any) builds are published to
 
         OutProperty: GauntletEditorTestArgs - Arguments to pass to Gauntlet tests executed under the editor
 		OutProperty: GauntletTargetTestArgs - Arguments to pass to Gauntlet tests executed using cooked builds
diff --git a/Engine/Build/Graph/Tasks/Inc/PathSettings.xml b/Engine/Build/Graph/Tasks/Inc/PathSettings.xml
--- a/Engine/Build/Graph/Tasks/Inc/PathSettings.xml
+++ b/Engine/Build/Graph/Tasks/Inc/PathSettings.xml
@@ -11,18 +11,29 @@
 
 		This script should be included after CommonProperties.xml where build strings are defines.
 
-		InOption: PublishBuilds - Enabled/Disable publishing of builds to build storage
+		InOption: PublishPackages - Enabled/Disable publishing of builds to build storage
 		
 		InProperty: ProjectName
 		InProperty: EscapedBranch
 
+		OutOption: PublishPackages - Publish packages from the build
+		OutOption: PublishStagedContent - Publish staged content from the build
+		OutOption: PublishReports - Publish reports from the build
+
 		OutProperty: ProjectOutputDirectory - Local path to package builds to
-		OutProperty: ProjectPublishPath - Path to publish builds to. If empty publishing should be skipped
+		OutProperty: NetworkOutputDirectory - Network path to be used to publish builds or staged content to
+		OutProperty: AutomationReportOutputDirectory - Path to publish reports to. If empty publishing should be skipped
+
 	-->
 
-	<!-- Publish builds to the network ?-->
-	<Option Name="PublishBuilds" DefaultValue="false" Description="Publish packaged builds to the builds folder on the netwok"/>
-	<Option Name="PublishReports" DefaultValue="$(PublishBuilds)" Description="Publish packaged builds to the builds folder on the netwok"/>
+	<!-- Publish packaged builds to the network ?-->
+	<Option Name="PublishPackages" DefaultValue="$(IsBuildMachine)" Description="Publish packaged builds to the builds folder on the network"/>
+
+	<!-- Publish staged content to the network-->
+	<Option Name="PublishStagedContent" DefaultValue="false" Description="When publishing builds, also include the staged content"/>
+
+	<!-- Publish automation reports to the network-->
+	<Option Name="PublishReports" DefaultValue="$(IsBuildMachine)" Description="Publish packaged builds to the builds folder on the netwok"/>
 
 
 	<!-- Where builds are saved to locally (todo - use ProjectDir?) -->
@@ -41,9 +52,24 @@
 	<Property Name="NetworkReportRoot" Value="/Volumes/Builds/Automation/Reports"/>
 	<Property Name="NetworkReportRoot" Value="P:/Builds/Automation/Reports" If="'$(HostPlatform)' == 'Win64'"/>
 
-	<!-- Declare the publish path, and point to the network if PublishBuilds is true -->
-	<Property Name="ProjectPublishPath" Value="" />
-	<Property Name="ProjectPublishPath" Value="$(NetworkPublishRoot)" If="$(PublishBuilds)"/>
+	<!-- Allow scripts to override the network locations-->
+	<Do If="'$(NetworkTempRootOverride)' != ''" >
+		<Property Name="NetworkTempRoot" Value="$(NetworkTempRootOverride)"/>
+	</Do>
+
+	<Do If="'$(NetworkPublishRootOverride)' != ''" >
+		<Property Name="NetworkPublishRoot" Value="$(NetworkPublishRootOverride)"/>
+	</Do>
+
+	<Do If="'$(NetworkReportRootOverride)' != ''" >
+		<Property Name="NetworkReportRoot" Value="$(NetworkReportRootOverride)"/>
+	</Do>
+
+	<!-- Declare the network outout path. Build machines need to write to the network for tests so use temp by default, but the publish path if that's being used -->
+	<Property Name="NetworkOutputDirectory" Value="" />
+	<Property Name="NetworkOutputDirectory" Value="$(NetworkTempRoot)" If="$(IsBuildMachine)"/>
+	<Property Name="NetworkOutputDirectory" Value="$(NetworkPublishRoot)" If="$(PublishPackages) or $(PublishStagedContent)"/>
+	<Property Name="NetworkOutputDirectory" Value="$(NetworkOutputDirectory)/$(EscapedBranch)/$(BuildName)" If="'$(NetworkOutputDirectory)' != ''"/>
 
 	<!-- Automation can publish reports to a folder that is mapped to a website -->
 	<Do If="$(IsBuildMachine) or $(PublishReports)" >
@@ -52,12 +78,4 @@
 		<Property Name="AutomationReportUrl" Value="http://automation.epicgames.net/reports/$(EscapedBranch)/$(ProjectName)/$(CL-String)"/>
 	</Do>
 
-	<!-- Build machines always need to publish somewhere so other agents can run the builds-->
-	<Do If="$(IsBuildMachine) And !$(PublishBuilds)">
-		<Property Name="ProjectPublishPath" Value="$(NetworkTempRoot)/$(EscapedBranch)/$(ProjectName)/$(CL-String)" />
-	</Do>
-
-	<!-- If we have a publish dir now, fully form the path-->
-	<Property Name="ProjectPublishPath" Value="$(NetworkPublishRoot)/$(EscapedBranch)/$(BuildName)" If="'$(ProjectPublishPath)' != ''"/>
-
 </BuildGraph>
\ No newline at end of file
diff --git a/Engine/Build/Graph/Tasks/Inc/PlatformOptions.xml b/Engine/Build/Graph/Tasks/Inc/PlatformOptions.xml
--- a/Engine/Build/Graph/Tasks/Inc/PlatformOptions.xml
+++ b/Engine/Build/Graph/Tasks/Inc/PlatformOptions.xml
@@ -126,6 +126,8 @@
 <!-- PS5 packaging isn't supported by Gauntlet yet :( )-->
 <Property Name="PlatformsToSkipPackaging" Value="PS5" />
 
+<!-- These platforms can't run from loose staged content-->
+<Property Name="PlatformsThatRequirePackaging" Value="Android+iOS" />
 
 <!-- Support explicit flags for horde but hide them otherwise as the option list is overwhelming..-->
 <Do If="$(ShowExplicitFlags)" >
diff --git a/Engine/Source/Programs/AutomationTool/Gauntlet/Framework/Gauntlet.Utils.cs b/Engine/Source/Programs/AutomationTool/Gauntlet/Framework/Gauntlet.Utils.cs
--- a/Engine/Source/Programs/AutomationTool/Gauntlet/Framework/Gauntlet.Utils.cs
+++ b/Engine/Source/Programs/AutomationTool/Gauntlet/Framework/Gauntlet.Utils.cs
@@ -183,7 +183,10 @@ namespace Gauntlet
 		/// </summary>
 		public static List<Action> PostAbortHandlers { get { return InnerPostAbortHandlers; } }
 
-
+		/// <summary>
+		/// Used by network functions that need to deal with multiple adapters
+		/// </summary>
+		public static string PreferredNetworkDomain {  get { return "epicgames.net";  } }
 	}
 
 	/// <summary>
@@ -1554,6 +1557,86 @@ namespace Gauntlet
 
 			return Files.ToArray();
 		}
+
+		/// <summary>
+		/// Returns true if any part of the specified path contains 'Component'. E.g. c:\Temp\Foo\Bar would return true for 'Foo'
+		/// </summary>
+		/// <param name="InPath">Path to search</param>
+		/// <param name="InComponent">Component to look for</param>
+		/// <param name="InPartialMatch">Whether a partial match is ok. E.g 'Fo' instead of 'Foo'</param>
+		/// <returns></returns>
+		public static bool PathContainsComponent(string InPath, string InComponent, bool InPartialMatch=false)
+		{
+			// normalize and split the path
+			var Components = new DirectoryInfo(InPath).FullName.Split(Path.DirectorySeparatorChar);
+
+			return InPartialMatch
+				? Components.Any(S => S.Contains(InComponent, StringComparison.OrdinalIgnoreCase))
+				: Components.Any(S => S.Equals(InComponent, StringComparison.OrdinalIgnoreCase));
+		}
+
+		/// <summary>
+		/// Returns subdirectories that match the specified regular expression, searching up to the specified maximum depth
+		/// </summary>
+		/// <param name="BaseDir">Directory to start searching in</param>
+		/// <param name="RegexPattern">Regular expression to match on</param>
+		/// <param name="RecursionDepth">Depth to search (-1 = unlimited)</param>
+		/// <returns></returns>
+		public static IEnumerable<DirectoryInfo> FindMatchingDirectories(string BaseDir, string RegexPattern, int RecursionDepth=0)
+		{
+			List<DirectoryInfo> Found = new List<DirectoryInfo>();
+
+			List<DirectoryInfo> Candidates = new List<DirectoryInfo>(new DirectoryInfo(BaseDir).GetDirectories());
+
+			Regex Pattern = new Regex(RegexPattern, RegexOptions.IgnoreCase);
+
+			int CurrentDepth = 0;
+
+			do
+			{
+				IEnumerable<DirectoryInfo> MatchingDirs = Candidates.Where(D => Pattern.IsMatch(D.Name));
+
+				Found.AddRange(MatchingDirs);
+
+				// recurse
+				Candidates = Candidates.SelectMany(D => D.GetDirectories()).ToList();
+
+			} while (Candidates.Any() && (CurrentDepth++ < RecursionDepth || RecursionDepth == -1)); ;
+
+			return Found;
+		}
+
+		/// <summary>
+		/// Returns files that match the specified regular expression in the provided folder, searching up to a maximum depth
+		/// </summary>
+		/// <param name="BaseDir">Directory to start searching in</param>
+		/// <param name="RegexPattern">Regular expression to match on</param>
+		/// <param name="RecursionDepth">Depth to search (-1 = unlimited)</param>
+		/// <returns></returns>
+		public static IEnumerable<FileInfo> FindMatchingFiles(string BaseDir, string RegexPattern, int RecursionDepth = 0)
+		{
+			List<FileInfo> Found = new List<FileInfo>();
+
+			IEnumerable<DirectoryInfo> CandidateDirs = new DirectoryInfo[] { new DirectoryInfo(BaseDir) };
+
+			Regex Pattern = new Regex(RegexPattern, RegexOptions.IgnoreCase);
+
+			int CurrentDepth = 0;
+
+			do
+			{
+				// check for matching files in this set of directories
+				IEnumerable<FileInfo> MatchingFiles = CandidateDirs.SelectMany(D => D.GetFiles()).Where(F => Pattern.IsMatch(F.Name));
+
+				Found.AddRange(MatchingFiles);
+
+				// recurse into this set of directories
+				CandidateDirs = CandidateDirs.SelectMany(D => D.GetDirectories()).ToList();
+
+			} while (CandidateDirs.Any() && (CurrentDepth++ < RecursionDepth || RecursionDepth == -1)); ;
+
+			return Found;
+		}
 	}
 
 	public static class FileUtils
