Modified BuildAndTest project script to allow projects to save hashed staging data to the network.

Added options for Linux to BuildAndTest and UE5 Horde templates for ShooterGame/QAGame for completeness.

diff --git a/Engine/Build/Graph/Tasks/BuildAndTestProject.xml b/Engine/Build/Graph/Tasks/BuildAndTestProject.xml
--- a/Engine/Build/Graph/Tasks/BuildAndTestProject.xml
+++ b/Engine/Build/Graph/Tasks/BuildAndTestProject.xml
@@ -15,10 +15,9 @@
 	<Option Name="ExtraCookArguments" DefaultValue="" Restrict="" Description="Extra arguments to use when cooking" />
 	<Option Name="ExtraStageAndPackageArguments" DefaultValue="" Restrict="" Description="Extra arguments to use when staging and packaging" />
 
-	<Option Name="PlatformsToSkipTests" DefaultValue="HoloLens+IOS+Lumin+Stadia+TVOS" Restrict=".*" Description="List of platforms that tests won't be run on" />
-
 	<Option Name="UseIncrementalAgents" DefaultValue="false" Restrict="true|false" Description="Use incremental agents for building" />
 
+
 	<Option Name="WithBATDefaults" DefaultValue="true" Description="Used by parent scripts to specify defaults when including this script"/>
 
 	<!-- 
@@ -219,26 +218,36 @@
 		<Property Name="PlatformCompileNodeName" Value="Compile $(TargetName) $(TargetPlatform)" />
 		<Property Name="PlatformCookNodeName" Value="Cook $(DisplayTargetName) $(TargetPlatform)" />
 		<Property Name="PlatformPackageNodeName" Value="Package $(DisplayTargetName) $(TargetPlatform)" />
+		<Property Name="PlatformHashNodeName" Value="Hash $(DisplayTargetName) $(TargetPlatform)" />
 		<Property Name="PlatformPublishNodeName" Value="Publish $(DisplayTargetName) $(TargetPlatform)" />
 
 		<!-- These will be a generated string of all configurations and builds/cooks/packages/tests for this platform that can be tied to aggregates -->
 		<Property Name="AllPlatformBuildNodes" Value="" />
 		<Property Name="AllPlatformCookNodes" Value="" />
 		<Property Name="AllPlatformPackageNodes" Value="" />
+		<Property Name="AllPlatformHashNodes" Value="" />
+
 		<Property Name="AllPlatformPublishNodes" Value="" />
 		<Property Name="AllPlatformBootTestNodes" Value="" />
 		<Property Name="AllPlatformTestNodes" Value="" />
 		
 		<!-- Set CookPlatform since it isn't always the same as the target -->
 		<Property Name="CookPlatform" Value="$(TargetPlatform)" />
-		<Property Name="CookPlatform" Value="Android_ETC2" If="'$(TargetPlatform)' == 'Android'"/>
 		<Property Name="CookPlatform" Value="Windows" If="'$(TargetPlatform)' == 'Win64'"/>
+		<Property Name="CookPlatform" Value="$(CookPlatform)_ETC2" If="'$(TargetPlatform)' == 'Android'"/>
 		<Property Name="CookPlatform" Value="$(CookPlatform)Client" If="'$(TargetType)' == 'Client'" />
 
+		<Property Name="BCRPlatform" Value="$(TargetPlatform)" />
+
+		<Property Name="BCRCookFlavor" Value="" If="" />
+		<Property Name="BCRCookFlavor" Value="ETC2" If="'$(TargetPlatform)' == 'Android'" />
+		<Property Name="StagedPlatformFolder" Value="$(CookPlatform)"  />
+		<!-- <Property Name="StagedPlatformFolder" Value="$(StagedPlatformFolder)Client" If="'$(TargetType)' == 'Client'" /> -->
+
 		<!-- Build the client and if necessary a server-->
-		<Agent Name="Build $(PlatformAgentName)" Type="$(BuildAgentType)">
+		<Agent Name="Make $(PlatformAgentName)" Type="$(BuildAgentType)">
 			<ForEach Name="TargetConfiguration" Values="$(TargetConfigurations)" Separator="+">
-				<Node Name="$(PlatformCompileNodeName) $(TargetConfiguration)" Requires="$(PlatformToolsNodeName)">
+				<Node Name="$(PlatformCompileNodeName) $(TargetConfiguration)" Requires="$(PlatformToolsNodeName)" RunEarly="true">
 					<Compile Target="$(TargetName)" Platform="$(TargetPlatform)" Configuration="$(TargetConfiguration)" Arguments="$(ExtraTargetCompileArguments)"/>
 					<Do If="'$(DedicatedServer)'" >
 						<Compile Target="$(TargetName)Server" Platform="Win64" Configuration="$(TargetConfiguration)" Arguments="$(ExtraTargetCompileArguments)"/>
@@ -246,50 +255,67 @@
 				</Node>
 				<Property Name="AllPlatformBuildNodes" Value="$(AllPlatformBuildNodes);$(PlatformCompileNodeName) $(TargetConfiguration)" />
 			</ForEach>
-		</Agent>
 
-		<!-- Cook for the platform. This can run at the same time as the above since we don't need the exe -->
-		<Agent Name="Cook $(PlatformAgentName)" Type="$(HostAgentType)">
+			<!-- Cook for the platform. This can run at the same time as the above since we don't need the exe -->
 			<Node Name="$(PlatformCookNodeName)" Requires="$(PlatformEditorNodeName)">
-				<Property Name="CookArgs" Value="$(CookPlatform)" />				
+				<Property Name="CookArgs" Value="" />				
 				<Property Name="CookArgs" Value="-MapIniSection=$(MapIniSection) $(CookPlatform)" If="'$(MapIniSection)' != ''"/>				
 				<Cook Project="$(TargetProject)" Platform="$(CookPlatform)" Arguments="$(CookArgs)"/>				
 			</Node>			
 			<Property Name="AllPlatformCookNodes" Value="$(AllPlatformCookNodes);$(PlatformCookNodeName)" />
-		</Agent>
 
-		<!-- Now package the exe and cooked data-->
-		<Agent Name="Package $(PlatformAgentName)" Type="$(HostAgentType)">
+			<!-- Do we publish a hashed version of the staged data? We only want to do this once per platform so use a local var that we'll flip to false later -->
+			<Property Name="DoHashedStagingPublish" Value="$(HostAgentType) == 'Win64'" />
+
+			<!-- Now package the exe and cooked data-->
 			<ForEach Name="PackageConfiguration" Values="$(PackageConfigurations)" Separator="+">
-				<Property Name="PackageDir" Value="$(ProjectOutputDirectory)/$(PackageConfiguration)/Packages" />
-				<Property Name="StagedDir" Value="$(ProjectOutputDirectory)/$(PackageConfiguration)/Staged" />
+				<Property Name="PackageDir" Value="$(ProjectPath)/Saved/Packages" />
+				<Property Name="StagedDir" Value="$(ProjectPath)/Saved/StagedBuilds" />
 
 				<Property Name="IsPackaged" Value="true" />
 				<Property Name="IsPackaged" Value="false" If="ContainsItem('$(PlatformsToSkipPackaging)','$(TargetPlatform)','+')" />	
 
+				<!-- Default BuildCookRun args needed staging/packaging/archiving this configuration -->
+				<Property Name="BCRArgs" Value="-project=$(TargetProject) -platform=$(BCRPlatform) -cookflavor=&quot;$(BCRCookFlavor)&quot; -configuration=$(PackageConfiguration) $(ExtraStageAndPackageArguments) -skipbuild -skipcook" />
+				<Property Name="BCRArgs" Value="$(BCRArgs) -client" If="'$(TargetType)' == 'Client'" />
+				<Property Name="BCRArgs" Value="$(BCRArgs) -server" If="'$(DedicatedServer)'" />
+
+				<!-- Pak stage and, if appropriate, package the build using BCR -->
 				<Node Name="$(PlatformPackageNodeName) $(PackageConfiguration)" Requires="$(PlatformCompileNodeName) $(PackageConfiguration);$(PlatformCookNodeName)">
-						<Property Name="BCRArgs" Value="-project=$(TargetProject) -platform=$(TargetPlatform) -configuration=$(PackageConfiguration) $(ExtraStageAndPackageArguments) -skipbuild -skipcook -pak -stage -StagingDirectory=&quot;$(StagedDir)&quot;" />
-						<Property Name="BCRArgs" Value="$(BCRArgs) -client" If="'$(TargetType)' == 'Client'" />
-						<Property Name="BCRArgs" Value="$(BCRArgs) -server" If="'$(DedicatedServer)'" />
-						<Property Name="BCRArgs" Value="$(BCRArgs) -package -archive -archivedirectory=&quot;$(PackageDir)&quot;" If="$(IsPackaged)" />		
-						<Command Name="BuildCookRun" Arguments="$(BCRArgs)" />
+						<Property Name="BCRPackageArgs" Value="$(BCRArgs) -pak -stage" />		
+						<Property Name="BCRPackageArgs" Value="$(BCRPackageArgs) -package" If="$(IsPackaged)" />		
+						<Command Name="BuildCookRun" Arguments="$(BCRPackageArgs)" />
 				</Node>
+
+				<!-- hash the content (we'll only tie one of these into our nodes, but always declare them as per best practices)-->
+				<Node Name="$(PlatformHashNodeName) $(PackageConfiguration)" Requires="$(PlatformPackageNodeName) $(PackageConfiguration)">
+					<Spawn Exe="$(RootDir)/Engine/Restricted/NotForLicensees/Binaries/Win64/unsync.exe" Arguments="hash $(StagedDir)/$(StagedPlatformFolder) -v"/>
+				</Node>				
+		
+				<!-- If publishing, we'll c-->
 				<Node Name="$(PlatformPublishNodeName) $(PackageConfiguration)" Requires="$(PlatformPackageNodeName) $(PackageConfiguration)">
 					<!-- 	
-						If this path is set we need to write to the network. This path will reflect set to scratch or semi-permanent storage based on whether
-						PublishPackages=true or we're a buld machine ( build machines always need to put builds on the network for testing)
-					-->
-					<Do If="'$(NetworkOutputDirectory)' != ''">						
-						<Do If="$(IsPackaged)" >
-							<Copy Files="..." From="$(PackageDir)/$(CookPlatform)" To="$(NetworkOutputDirectory)/$(CookPlatform)/$(PackageConfiguration)"/>
-						</Do>
-						<Do If="$(PublishStagedContent) or !$(IsPackaged)" >
-							<Copy Files="..." From="$(StagedDir)/$(CookPlatform)" To="$(NetworkOutputDirectory)/$(CookPlatform)/Staged"/>
-						</Do>
+						This node will be tied in if PublishBuilds=true or DoHashedStagingPublish=true, in whicb case the network path will have been
+						set up in PathSettings.xml  This path will reflect set to scratch or semi-permanent storage based on whether
+						PublishBuilds=true or we're a buld machine (build machines always need to put builds on the network for testing)
+					-->					
+					<Do If="$(IsPackaged) and $(PublishBuilds)" >
+						<Property Name="BCRPublishArgs" Value="$(BCRArgs) -skippak -skipstage -skippackage -archive -archivedirectory=&quot;$(NetworkOutputDirectory)/$(StagedPlatformFolder)/$(PackageConfiguration)&quot;" />		
+						<Command Name="BuildCookRun" Arguments="$(BCRPublishArgs)" />
+					</Do>
+					<!-- We'll hash-->
+					<Do If="!$(IsPackaged) and ($(PublishBuilds) or $(DoHashedStagingPublish))" >
+						<Copy Files="..." From="$(StagedDir)/$(StagedPlatformFolder)" To="$(NetworkOutputDirectory)/$(StagedPlatformFolder)/Staged"/>
+					</Do>
+					<Do If="$(IsPackaged) and $(DoHashedStagingPublish)" >
+						<Copy Files="..." From="$(StagedDir)/$(StagedPlatformFolder)" To="$(NetworkOutputDirectory)/$(StagedPlatformFolder)/Staged"/>
 					</Do>
 				</Node>
-				<Property Name="AllPlatformCookNodes" Value="$(AllPlatformCookNodes);$(PlatformPackageNodeName) $(PackageConfiguration)" />
+				<!-- Tie these nodes into the list for this platform. They will later be tied in (or not) based on overall options -->
+				<Property Name="AllPlatformPackageNodes" Value="$(AllPlatformPackageNodes);$(PlatformPackageNodeName) $(PackageConfiguration)" />
 				<Property Name="AllPlatformPublishNodes" Value="$(AllPlatformPublishNodes);$(PlatformPublishNodeName) $(PackageConfiguration)" />
+				<Property Name="AllPlatformHashNodes" Value="$(AllPlatformHashNodes);$(PlatformHashNodeName) $(PackageConfiguration)" If="$(DoHashedStagingPublish)" />
+				<Property Name="DoHashedStagingPublish" Value="false" />
 			</ForEach>
 		</Agent>			
 
@@ -334,7 +360,9 @@
 						<Command Name="RunUnreal" Arguments="-project=$(TargetProject) -test=&quot;$(TestName)&quot; -platform=$(TargetPlatform) -configuration=$(TestConfiguration) $(GauntletBuildArg) $(GauntletTargetTestArgs) $(GauntletReportArgs)"/>
 					</Node>
 					<!-- Add it to our list -->
-					<Property Name="AllPlatformTestNodes" Value="$(AllPlatformTestNodes);$(TestNodeName)"/>
+					<Do If="!ContainsItem('$(PlatformsToSkipTests)', '$(TargetPlatform)', '+')">
+						<Property Name="AllPlatformTestNodes" Value="$(AllPlatformTestNodes);$(TestNodeName)"/>
+					</Do>
 				</ForEach>	
 			</ForEach>
 		</Agent>
@@ -342,7 +370,7 @@
 		<!-- Declare labels for CIS -->
 		<Label Category="Platforms" Name="$(ProjectName) $(TargetPlatform)" Requires="$(AllPlatformBuildNodes)" />
 		<Label Category="Cook" Name="Cook $(DisplayTargetName) $(TargetPlatform)" Requires="$(AllPlatformCookNodes)" Exclude="$(PlatformEditorNodeName)" />
-		<Label Category="Package" Name="Pkg $(DisplayTargetName) $(TargetPlatform)" Requires="$(AllPlatformPackageNodes);$(AllPlatformPublishNodes)" Exclude="$(AllPlatformBuildNodes);$(AllPlatformCookNodes)"/>
+		<Label Category="Package" Name="Pkg $(DisplayTargetName) $(TargetPlatform)" Requires="$(AllPlatformPackageNodes);$(AllPlatformHashNodes);$(AllPlatformPublishNodes)" Exclude="$(AllPlatformBuildNodes);$(AllPlatformCookNodes)"/>
 		<Label Category="Test" Name="BootTest $(DisplayTargetName) $(TargetPlatform)" Requires="$(AllPlatformBootTestNodes)" Exclude="$(PlatformEditorNodeName)" />
 		<Label Category="Test" Name="Test $(DisplayTargetName) $(TargetPlatform)" Requires="$(AllPlatformTestNodes)" Exclude="$(PlatformEditorNodeName);$(AllPlatformBootTestNodes)"/>
 
@@ -350,7 +378,8 @@
 		<Property Name="PlatformTasks" Value="" />
 		<Property Name="PlatformTasks" Value="$(AllPlatformBuildNodes)" If="!$(SkipBuild)" />
 		<Property Name="PlatformTasks" Value="$(PlatformTasks);$(AllPlatformCookNodes)" If="!$(SkipCook) and !$(SkipPublish)" />		
-		<Property Name="PlatformTasks" Value="$(PlatformTasks);$(AllPlatformPackageNodes);$(AllPlatformPublishNodes)" If="!$(SkipPublish)" />	
+		<Property Name="PlatformTasks" Value="$(PlatformTasks);$(AllPlatformPackageNodes);$(AllPlatformPublishNodes)" If="!$(SkipPublish) And $(PublishBuilds)" />	
+		<Property Name="PlatformTasks" Value="$(PlatformTasks);$(AllPlatformHashNodes)"  If="$(PublishHashedStagingData)" />	
 		<Property Name="PlatformTasks" Value="$(PlatformTasks);$(AllPlatformBootTestNodes);$(AllPlatformTestNodes)" If="!ContainsItem('$(PlatformsToSkipTests)','$(TargetPlatform)','+')" />		
 
 		<Property Name="GeneratedTargetNodes" Value="$(GeneratedTargetNodes);$(PlatformTasks)"/>
diff --git a/Engine/Build/Graph/Tasks/Inc/PathSettings.xml b/Engine/Build/Graph/Tasks/Inc/PathSettings.xml
--- a/Engine/Build/Graph/Tasks/Inc/PathSettings.xml
+++ b/Engine/Build/Graph/Tasks/Inc/PathSettings.xml
@@ -11,13 +11,13 @@
 
 		This script should be included after CommonProperties.xml where build strings are defines.
 
-		InOption: PublishPackages - Enabled/Disable publishing of builds to build storage
+		InOption: PublishBuilds - Enabled/Disable publishing of builds to build storage
 		
 		InProperty: ProjectName
 		InProperty: EscapedBranch
 
-		OutOption: PublishPackages - Publish packages from the build
-		OutOption: PublishStagedContent - Publish staged content from the build
+		OutOption: PublishBuilds - Publish output from builds to a network share
+		OutOption: PublishHashedStagingData - Publish staged content with hashes for quick syncing
 		OutOption: PublishReports - Publish reports from the build
 
 		OutProperty: ProjectOutputDirectory - Local path to package builds to
@@ -27,10 +27,10 @@
 	-->
 
 	<!-- Publish packaged builds to the network ?-->
-	<Option Name="PublishPackages" DefaultValue="$(IsBuildMachine)" Description="Publish packaged builds to the builds folder on the network"/>
+	<Option Name="PublishBuilds" DefaultValue="$(IsBuildMachine)" Description="Publish packaged builds to the builds folder on the network"/>
 
-	<!-- Publish staged content to the network-->
-	<Option Name="PublishStagedContent" DefaultValue="false" Description="When publishing builds, also include the staged content"/>
+	<!-- Publish a hashed staging data that developers can quickly sync to -->
+	<Option Name="PublishHashedStagingData" DefaultValue="false" Restrict="true|false" Description="Create hashes for syncing via Unsync" />
 
 	<!-- Publish automation reports to the network-->
 	<Option Name="PublishReports" DefaultValue="$(IsBuildMachine)" Description="Publish packaged builds to the builds folder on the netwok"/>
@@ -64,11 +64,12 @@
 	<Do If="'$(NetworkReportRootOverride)' != ''" >
 		<Property Name="NetworkReportRoot" Value="$(NetworkReportRootOverride)"/>
 	</Do>
+	<!-- End network overrides -->
 
 	<!-- Declare the network outout path. Build machines need to write to the network for tests so use temp by default, but the publish path if that's being used -->
 	<Property Name="NetworkOutputDirectory" Value="" />
 	<Property Name="NetworkOutputDirectory" Value="$(NetworkTempRoot)" If="$(IsBuildMachine)"/>
-	<Property Name="NetworkOutputDirectory" Value="$(NetworkPublishRoot)" If="$(PublishPackages) or $(PublishStagedContent)"/>
+	<Property Name="NetworkOutputDirectory" Value="$(NetworkPublishRoot)" If="$(PublishBuilds) or $(PublishHashedStagingData)"/>
 	<Property Name="NetworkOutputDirectory" Value="$(NetworkOutputDirectory)/$(EscapedBranch)/$(BuildName)" If="'$(NetworkOutputDirectory)' != ''"/>
 
 	<!-- Automation can publish reports to a folder that is mapped to a website -->
diff --git a/Engine/Build/Graph/Tasks/Inc/PlatformOptions.xml b/Engine/Build/Graph/Tasks/Inc/PlatformOptions.xml
--- a/Engine/Build/Graph/Tasks/Inc/PlatformOptions.xml
+++ b/Engine/Build/Graph/Tasks/Inc/PlatformOptions.xml
@@ -37,14 +37,19 @@
 <Option Name="PackageConfigurations" DefaultValue="$(TargetConfigurations)" Description="Configuration to package"/>
 <Option Name="TestConfigurations" DefaultValue="$(PackageConfigurations)" Description="Configuration to test"/>
 
+<!-- No means of reliably running these tests atm -->
+<Option Name="PlatformsToSkipTests" DefaultValue="HoloLens+IOS+Linux+Lumin+Stadia+TVOS" Restrict=".*" Description="List of platforms that tests won't be run on" />
+
 <!-- Support explicit flags for horde but hide them otherwise as the option list is overwhelming..-->
 <Do If="$(ShowExplicitFlags)" >
 	<Option Name="WithWin64Editor" DefaultValue="false" Description="Build and test Win64 Editor"/>
 	<Option Name="WithMacEditor" DefaultValue="false" Description="Build and test Mac Editor"/>
+	<Option Name="WithLinuxEditor" DefaultValue="false" Description="Build and test Linux Editor"/>
 
 	<Option Name="WithAndroid" DefaultValue="false" Description="Build and test Android"/>
 	<Option Name="WithHoloLens" DefaultValue="false" Description="Build and test HoloLens"/>
 	<Option Name="WithIOS" DefaultValue="false" Description="Build and test IOS"/>
+	<Option Name="WithLinux" DefaultValue="false" Description="Build and test Linux"/>
 	<Option Name="WithLumin" DefaultValue="false" Description="Build and test Lumin"/>
 	<Option Name="WithMac" DefaultValue="false" Description="Build and test Mac"/>
 	<Option Name="WithPS4" DefaultValue="false" Description="Build and test PS4"/>
@@ -61,6 +66,7 @@
 		<Property Name="EditorPlatforms" Value="" />
 		<Property Name="EditorPlatforms" Value="$(EditorPlatforms)Win64+" If="$(WithWin64Editor)" />
 		<Property Name="EditorPlatforms" Value="$(EditorPlatforms)Mac+" If="$(WithMacEditor)" />
+		<Property Name="EditorPlatforms" Value="$(EditorPlatforms)Linux+" If="$(WithLinuxEditor)" />
 	</Do>
 
 	<Do If="'$(TargetPlatforms)' == ''" >
@@ -68,6 +74,7 @@
 		<Property Name="TargetPlatforms" Value="$(TargetPlatforms)Android+" If="$(WithAndroid)" />
 		<Property Name="TargetPlatforms" Value="$(TargetPlatforms)HoloLens+" If="$(WithHoloLens)" />
 		<Property Name="TargetPlatforms" Value="$(TargetPlatforms)IOS+" If="$(WithIOS)" />
+		<Property Name="TargetPlatforms" Value="$(TargetPlatforms)Linux+" If="$(WithLinux)" />
 		<Property Name="TargetPlatforms" Value="$(TargetPlatforms)Lumin+" If="$(WithLumin)" />
 		<Property Name="TargetPlatforms" Value="$(TargetPlatforms)Mac+" If="$(WithMac)" />
 		<Property Name="TargetPlatforms" Value="$(TargetPlatforms)PS4+" If="$(WithPS4)" />
@@ -124,7 +131,7 @@
 </ForEach>
 
 <!-- PS5 packaging isn't supported by Gauntlet yet :( )-->
-<Property Name="PlatformsToSkipPackaging" Value="PS5" />
+<Property Name="PlatformsToSkipPackaging" Value="Win64+Linux+PS5" />
 
 <!-- These platforms can't run from loose staged content-->
 <Property Name="PlatformsThatRequirePackaging" Value="Android+iOS" />
