Add support for overriding agent types.

#rnx

diff --git a/Engine/Build/Graph/Tasks/BuildAndTestProject.xml b/Engine/Build/Graph/Tasks/BuildAndTestProject.xml
--- a/Engine/Build/Graph/Tasks/BuildAndTestProject.xml
+++ b/Engine/Build/Graph/Tasks/BuildAndTestProject.xml
@@ -16,7 +16,8 @@
 	<Option Name="ExtraStageAndPackageArguments" DefaultValue="" Restrict="" Description="Extra arguments to use when staging and packaging" />
 
 	<Option Name="UseIncrementalAgents" DefaultValue="false" Restrict="true|false" Description="Use incremental agents for building" />
-
+	<Option Name="AgentOverride" DefaultValue="" Description="If set, other logic is ignored and this agent type is used for all non-test work." />
+	<Option Name="TestAgentOverride" DefaultValue="" Description="If set, other logic is ignored and this agent type is used for all test work." />
 
 	<Option Name="WithBATDefaults" DefaultValue="true" Description="Used by parent scripts to specify defaults when including this script"/>
 
@@ -131,6 +132,9 @@
 		<!-- Use incremental agents for building? -->
 		<Property Name="HostAgentType" Value="Incremental$(HostAgentType)" If="$(UseIncrementalAgents)" />
 		<Property Name="TestAgentType" Value="Test$(EditorPlatform)" />
+		<!-- Use overrides for agent types? -->
+		<Property Name="HostAgentType" Value="$(AgentOverride)" If="'$(AgentOverride)' != ''" />
+		<Property Name="TestAgentType" Value="$(TestAgentOverride)"  If="'$(TestAgentOverride)' != ''" />
 
 		<!-- Test nodes for this editor platform -->
 		<Property Name="PlatformEditorTestNodes" Value="" />
@@ -210,6 +214,10 @@
 		<Property Name="BuildAgentType" Value="$(HostAgentType)" />
 		<Property Name="BuildAgentType" Value="Incremental$(HostAgentType)" If="$(UseIncrementalAgents)" />
 
+		<!-- Use agent types from override? -->
+		<Property Name="BuildAgentType" Value="$(AgentOverride)" If="'$(AgentOverride)' != ''"/>
+		<Property Name="TestAgentType" Value="$(TestAgentOverride)" If="'$(TestAgentOverride)' != ''"/>
+
 		<!-- Building and cooking will require these nodes from the appropriate host. They should have been built above... -->
 		<Property Name="PlatformToolsNodeName" Value="Compile $(ProjectName)Tools $(HostAgentType)" />
 		<Property Name="PlatformEditorNodeName" Value="Compile $(ProjectName)Editor $(HostAgentType)" />
