Remove PS5 from "PlatformsToSkipPackaging" to allow PS5 packaging on Frosty

diff --git a/Engine/Build/Graph/Tasks/Inc/PlatformOptions.xml b/Engine/Build/Graph/Tasks/Inc/PlatformOptions.xml
--- a/Engine/Build/Graph/Tasks/Inc/PlatformOptions.xml
+++ b/Engine/Build/Graph/Tasks/Inc/PlatformOptions.xml
@@ -131,7 +131,7 @@
 </ForEach>
 
 <!-- PS5 packaging isn't supported by Gauntlet yet :( )-->
-<Property Name="PlatformsToSkipPackaging" Value="Win64+Linux+PS5" />
+<Property Name="PlatformsToSkipPackaging" Value="Win64+Linux" />
 
 <!-- These platforms can't run from loose staged content-->
 <Property Name="PlatformsThatRequirePackaging" Value="Android+iOS" />
