BuildEditorAndTools.xml:

Recursively bundle pre-compiled AutomationTool binaries, based on the method used in InstalledEngineBuild.xml

(AutomationTool is pre-built before the script is run - there is no need to build it again)

diff --git a/Engine/Build/Graph/Examples/BuildEditorAndTools.xml b/Engine/Build/Graph/Examples/BuildEditorAndTools.xml
--- a/Engine/Build/Graph/Examples/BuildEditorAndTools.xml
+++ b/Engine/Build/Graph/Examples/BuildEditorAndTools.xml
@@ -123,12 +123,8 @@
 			<Property Name="ArchiveDir" Value="$(RootDir)\LocalBuilds\ArchiveForUGS"/>
 			<Delete Files="$(ArchiveDir)\..."/>
 
-			<!-- Tag required files from UAT and UBT that will already have been built -->
-			<Tag Files="Engine/Source/Programs/AutomationTool/..." Filter="*.csproj" With="#UAT Projects"/>
-			<Tag Files="Engine/Platforms/..." Filter="*.Automation.csproj" With="#UAT Projects"/>
-			<Tag Files="$(RootDir)/.../Build/Scripts/..." Filter="*.Automation.csproj" With="#UAT Projects"/>
-			<CsCompile Project="#UAT Projects" Configuration="Development" 
-Platform="AnyCPU"  Properties="EngineDir=$(RootDir)/Engine" Tag="#ArchiveFiles" EnumerateOnly="true"/>
+			<!-- Tag AutomationTool build folder recursively -->
+			<Tag Files="Engine/Binaries/DotNET/AutomationTool/..." With="#ArchiveBinaries"/>
 			
 			<!-- Write a text file marker indicating that we're using precompiled binaries -->
 			<Spawn Exe="cmd.exe" Arguments="/C echo. &gt;&quot;$(RootDir)/Engine/Build/PrecompiledBinaries.txt&quot;"/>
