Added build graph option to update version files

diff --git a/Engine/Build/Graph/Tasks/BuildAndTestProject.xml b/Engine/Build/Graph/Tasks/BuildAndTestProject.xml
--- a/Engine/Build/Graph/Tasks/BuildAndTestProject.xml
+++ b/Engine/Build/Graph/Tasks/BuildAndTestProject.xml
@@ -46,6 +46,7 @@
 	<Option Name="DedicatedServer" DefaultValue="false" Description="Project uses a dedicated server" />
 
 	<!-- Customizations -->
+	<Option Name="Versioned" Restrict="true|false" DefaultValue="false" Description="Whether to embed changelist number into binaries"/>
 	<Option Name="SkipBuild" DefaultValue="false" Description="Whether to skip building" />
 	<Option Name="SkipCook" DefaultValue="false" Description="Whether to skip cooking and packaging" />
 	<Option Name="SkipPackage" DefaultValue="false" Description="Whether to skip packaging" />
@@ -140,8 +141,12 @@
 		<Property Name="PlatformEditorTestNodes" Value="" />
 
 		<Agent Name="Build $(ProjectName)Editor $(EditorPlatform)" Type="$(HostAgentType)">
+			<!-- Update the Version.h and Build.version files if we're making a formal build on a build machine -->
+			<Node Name="Update Version Files">
+				<SetVersion Change="$(Change)" Branch="$(EscapedBranch)" If="$(Versioned) and $(IsBuildMachine)"/>
+			</Node>
 			<!-- First generate nodes to build the required tools (e.g. UHT, ShaderCompiler, Editor -->
-			<Node Name="Compile $(ProjectName)Tools $(EditorPlatform)">
+			<Node Name="Compile $(ProjectName)Tools $(EditorPlatform)" Requires="Update Version Files">
 				<Compile Target="UnrealHeaderTool" Platform="$(EditorPlatform)" Configuration="Development"/>
 				<Compile Target="ShaderCompileWorker" Platform="$(EditorPlatform)" Configuration="Development"/>
 				<Compile Target="UnrealPak" Platform="$(EditorPlatform)" Configuration="Development"/>
