Added options to specify editor/target platforms to skip tests and separated them between boot tests and regular tests

diff --git a/Engine/Build/Graph/Tasks/BuildAndTestProject.xml b/Engine/Build/Graph/Tasks/BuildAndTestProject.xml
--- a/Engine/Build/Graph/Tasks/BuildAndTestProject.xml
+++ b/Engine/Build/Graph/Tasks/BuildAndTestProject.xml
@@ -42,6 +42,14 @@
 	<!-- This will define options for all platforms and create lists called EditorPlatforms, RequiredEditorPlatforms, and TargetPlatforms. E.g. TargetPlatforms=Win64+Mac+PS4 -->
 	<Include Script="Inc/PlatformOptions.xml" />
 
+	<!-- For compatibility with build graphs relying on PlatformsToSkipTests option from PlatformOptions that was replaced with Editor/Target specific options -->
+	<Do If="'$(PlatformsToSkipTests)' != ''" >
+		<Property Name="EditorPlatformsToSkipTests" Value="$(PlatformsToSkipTests)" />
+		<Property Name="TargetPlatformsToSkipTests" Value="$(PlatformsToSkipTests)" />
+		<Property Name="EditorPlatformsToSkipBootTests" Value="$(PlatformsToSkipTests)" />
+		<Property Name="TargetPlatformsToSkipBootTests" Value="$(PlatformsToSkipTests)" />
+	</Do>
+
 	<!-- Does this project require a dedicated server for testing? (not currently implemented!)-->
 	<Option Name="DedicatedServer" DefaultValue="false" Description="Project uses a dedicated server" />
 
@@ -57,8 +65,10 @@
 	<Option Name="EditorBootTest" DefaultValue="!$(SkipTest)" Description="Run a simple boot test on the editor before other tests"/>
 	<Option Name="TargetBootTest" DefaultValue="!$(SkipTest)" Description="Run a simple boot test on target platforms before other tests"/>
 
-	<Option Name="EditorTestList" DefaultValue="$(DefaultEditorTestList)" Description="Tests to run on the edutor. Test1+Test2 etc"/>
+	<Option Name="EditorTestList" DefaultValue="$(DefaultEditorTestList)" Description="Tests to run on the editor. Test1+Test2 etc"/>
 	<Option Name="TargetTestList" DefaultValue="$(DefaultTargetTestList)" Description="Tests to run on the target platform(s). Test1+Test2 etc"/>
+	<Option Name="EditorPlatformsToSkipBootTests" DefaultValue="$(EditorPlatformsToSkipTests)" Description="Editor platforms that should skip boot tests."/>
+	<Option Name="TargetPlatformsToSkipBootTests" DefaultValue="$(TargetPlatformsToSkipTests)" Description="Target platforms that should skip boot tests."/>
 
 	<!-- Settings files that can be overriden -->
 	<Option Name="PathSettingsFile" DefaultValue="Inc/PathSettings.xml" Description="Path settings to use" />
@@ -98,7 +108,10 @@
 
 	<!-- If skip test is true, add all target platforms to the "do not test" list -->
 	<Do If="$(SkipTest)">
-		<Property Name="PlatformsToSkipTests" Value="$(EditorPlatforms)+$(TargetPlatforms)" />
+		<Property Name="EditorPlatformsToSkipTests" Value="$(EditorPlatforms)" />
+		<Property Name="TargetPlatformsToSkipTests" Value="$(TargetPlatforms)" />
+		<Property Name="EditorPlatformsToSkipBootTests" Value="$(EditorPlatformsToSkipTests)" />
+		<Property Name="TargetPlatformsToSkipBootTests" Value="$(TargetPlatformsToSkipTests)" />
 	</Do>
 
 	<!-- If skip packaging is true, add all target platforms to the "do not test" list -->
@@ -179,12 +192,12 @@
 			</Do>
 			<!-- Now generate a test node for each editor test in the list-->
 			<ForEach Name="TestName" Values="$(EditorTestList)" Separator="+">
-			<Do If="'$(AutomationReportOutputDirectory)' != ''">
-				<Property Name="OutputFragmentPath" Value="$(EditorPlatform)/Editor$(TestName)" />
-				<Property Name="OutputReportPath" Value="$(AutomationReportOutputDirectory)/$(OutputFragmentPath)" />
-				<Property Name="GauntletReportArgs" Value="-logdir=&quot;$(OutputReportPath)&quot; -artifactname=Gauntlet -ReportExportPath=&quot;$(OutputReportPath)&quot;" />
-				<Property Name="GauntletReportArgs" Value="$(GauntletReportArgs) -ReportURL=&quot;$(AutomationReportUrl)/$(OutputFragmentPath)/&quot;" If="'$(AutomationReportUrl)' != ''" />
-			</Do>
+				<Do If="'$(AutomationReportOutputDirectory)' != ''">
+					<Property Name="OutputFragmentPath" Value="$(EditorPlatform)/Editor$(TestName)" />
+					<Property Name="OutputReportPath" Value="$(AutomationReportOutputDirectory)/$(OutputFragmentPath)" />
+					<Property Name="GauntletReportArgs" Value="-logdir=&quot;$(OutputReportPath)&quot; -artifactname=Gauntlet -ReportExportPath=&quot;$(OutputReportPath)&quot;" />
+					<Property Name="GauntletReportArgs" Value="$(GauntletReportArgs) -ReportURL=&quot;$(AutomationReportUrl)/$(OutputFragmentPath)/&quot;" If="'$(AutomationReportUrl)' != ''" />
+				</Do>
 				<Property Name="TestNodeName" Value="$(ProjectName)Editor $(EditorPlatform) Test=$(TestName)" />
 				<Node Name="$(TestNodeName)" Requires="$(TestPrerequisites)">
 					<Command Name="RunUnreal" Arguments="-project=$(TargetProject) -build=Editor -test=&quot;$(TestName)&quot; $(GauntletReportArgs) $(GauntletEditorTestArgs)" />
@@ -197,7 +210,7 @@
 			<Label Category="Editors" Name="$(ProjectName) Editor $(EditorPlatform)" Requires="$(GeneratedToolNodes)" />
 
 			<!-- If this editor was requested (as opposed to being required due to target platforms) then link in the test nodes and set up labels -->
-			<Do If="ContainsItem('$(EditorPlatforms)', '$(EditorPlatform)', '+') And !ContainsItem('$(PlatformsToSkipTests)', '$(EditorPlatform)', '+')">
+			<Do If="ContainsItem('$(EditorPlatforms)', '$(EditorPlatform)', '+') and !ContainsItem('$(EditorPlatformsToSkipTests)', '$(EditorPlatform)', '+')">
 				<Label Category="Test" Name="BootTest $(ProjectName)Editor $(EditorPlatform)" Requires="$(PlatformBootTestNodeName)" />
 				<Label Category="Test" Name="Test $(ProjectName)Editor $(EditorPlatform)" Requires="$(PlatformEditorTestNodes)" />
 				<Property Name="GeneratedEditorTestNodes" Value="$(GeneratedEditorTestNodes);$(PlatformBootTestNodeName);$(PlatformEditorTestNodes)" />
@@ -354,7 +367,7 @@
 					<Command Name="RunUnreal" Arguments="-project=$(TargetProject) -test=&quot;UE.BootTest&quot; -platform=$(TargetPlatform) -configuration=$(TestConfiguration) $(GauntletBuildArg) $(GauntletTargetTestArgs) $(GauntletReportArgs)" />
 				</Node>
 				<!-- Link the boot test in to the prerequisites and generated nodes if its enabled-->
-				<Do If="$(TargetBootTest)">
+				<Do If="$(TargetBootTest) and !ContainsItem('$(TargetPlatformsToSkipBootTests)', '$(TargetPlatform)', '+')">
 					<Property Name="TestPrerequisites" Value="$(TestPrerequisites);$(BootTestName)" />
 					<Property Name="AllPlatformBootTestNodes" Value="$(AllPlatformBootTestNodes);$(BootTestName)"/>
 				</Do>
@@ -373,7 +386,7 @@
 						<Command Name="RunUnreal" Arguments="-project=$(TargetProject) -test=&quot;$(TestName)&quot; -platform=$(TargetPlatform) -configuration=$(TestConfiguration) $(GauntletBuildArg) $(GauntletTargetTestArgs) $(GauntletReportArgs)"/>
 					</Node>
 					<!-- Add it to our list -->
-					<Do If="!ContainsItem('$(PlatformsToSkipTests)', '$(TargetPlatform)', '+')">
+					<Do If="!ContainsItem('$(TargetPlatformsToSkipTests)', '$(TargetPlatform)', '+')">
 						<Property Name="AllPlatformTestNodes" Value="$(AllPlatformTestNodes);$(TestNodeName)"/>
 					</Do>
 				</ForEach>	
@@ -393,7 +406,7 @@
 		<Property Name="PlatformTasks" Value="$(PlatformTasks);$(AllPlatformCookNodes)" If="!$(SkipCook) and !$(SkipPublish)" />		
 		<Property Name="PlatformTasks" Value="$(PlatformTasks);$(AllPlatformPackageNodes);$(AllPlatformPublishNodes)" If="!$(SkipPublish) And $(PublishBuilds)" />	
 		<Property Name="PlatformTasks" Value="$(PlatformTasks);$(AllPlatformHashNodes)"  If="$(PublishHashedStagingData)" />	
-		<Property Name="PlatformTasks" Value="$(PlatformTasks);$(AllPlatformBootTestNodes);$(AllPlatformTestNodes)" If="!ContainsItem('$(PlatformsToSkipTests)','$(TargetPlatform)','+')" />		
+		<Property Name="PlatformTasks" Value="$(PlatformTasks);$(AllPlatformBootTestNodes);$(AllPlatformTestNodes)" />		
 
 		<Property Name="GeneratedTargetNodes" Value="$(GeneratedTargetNodes);$(PlatformTasks)"/>
 	</ForEach>
diff --git a/Engine/Build/Graph/Tasks/Inc/PlatformOptions.xml b/Engine/Build/Graph/Tasks/Inc/PlatformOptions.xml
--- a/Engine/Build/Graph/Tasks/Inc/PlatformOptions.xml
+++ b/Engine/Build/Graph/Tasks/Inc/PlatformOptions.xml
@@ -38,7 +38,9 @@
 <Option Name="TestConfigurations" DefaultValue="$(PackageConfigurations)" Description="Configuration to test"/>
 
 <!-- No means of reliably running these tests atm -->
-<Option Name="PlatformsToSkipTests" DefaultValue="HoloLens+IOS+Linux+Lumin+Stadia+TVOS" Restrict=".*" Description="List of platforms that tests won't be run on" />
+<Option Name="EditorPlatformsToSkipTests" DefaultValue="" Restrict=".*" Description="List of editor platforms that editor tests won't be run on" />
+<Option Name="TargetPlatformsToSkipTests" DefaultValue="HoloLens+IOS+Linux+Lumin+Stadia+TVOS" Restrict=".*" Description="List of target platforms that target tests won't be run on" />
+<Option Name="PlatformsToSkipTests" DefaultValue="" Restrict=".*" Description="List of platforms that all tests won't be run on. Overrides Editor/TargetPlatformsToSkipTests when non-empty." />
 
 <!-- Support explicit flags for horde but hide them otherwise as the option list is overwhelming..-->
 <Do If="$(ShowExplicitFlags)" >
