Use buffered file IO during unsync manifest generation step

diff --git a/Engine/Build/Graph/Tasks/BuildAndTestProject.xml b/Engine/Build/Graph/Tasks/BuildAndTestProject.xml
--- a/Engine/Build/Graph/Tasks/BuildAndTestProject.xml
+++ b/Engine/Build/Graph/Tasks/BuildAndTestProject.xml
@@ -315,7 +315,7 @@
 
 				<!-- hash the content (we'll only tie one of these into our nodes, but always declare them as per best practices)-->
 				<Node Name="$(PlatformHashNodeName) $(PackageConfiguration)" Requires="$(PlatformPackageNodeName) $(PackageConfiguration)">
-					<Spawn Exe="$(RootDir)/Engine/Restricted/NotForLicensees/Binaries/Win64/unsync.exe" Arguments="hash $(StagedDir)/$(StagedPlatformFolder) -v"/>
+					<Spawn Exe="$(RootDir)/Engine/Restricted/NotForLicensees/Binaries/Win64/unsync.exe" Arguments="hash $(StagedDir)/$(StagedPlatformFolder) -v --buffered-files"/>
 				</Node>				
 		
 				<!-- If publishing, we'll c-->
