Fix local project output path when using BuildAndTestProject buildgraph

#rnx

diff --git a/Engine/Build/Graph/Tasks/BuildAndTestProject.xml b/Engine/Build/Graph/Tasks/BuildAndTestProject.xml
--- a/Engine/Build/Graph/Tasks/BuildAndTestProject.xml
+++ b/Engine/Build/Graph/Tasks/BuildAndTestProject.xml
@@ -295,8 +295,6 @@
 
 			<!-- Now package the exe and cooked data-->
 			<ForEach Name="PackageConfiguration" Values="$(PackageConfigurations)" Separator="+">
-				<Property Name="PackageDir" Value="$(ProjectPath)/Saved/Packages" />
-				<Property Name="StagedDir" Value="$(ProjectPath)/Saved/StagedBuilds" />
 
 				<Property Name="IsPackaged" Value="true" />
 				<Property Name="IsPackaged" Value="false" If="ContainsItem('$(PlatformsToSkipPackaging)','$(TargetPlatform)','+')" />	
@@ -315,8 +313,8 @@
 
 				<!-- hash the content (we'll only tie one of these into our nodes, but always declare them as per best practices)-->
 				<Node Name="$(PlatformHashNodeName) $(PackageConfiguration)" Requires="$(PlatformPackageNodeName) $(PackageConfiguration)">
-					<Spawn Exe="$(RootDir)/Engine/Restricted/NotForLicensees/Binaries/Win64/unsync.exe" Arguments="hash $(StagedDir)/$(StagedPlatformFolder) -v --buffered-files"/>
-				</Node>				
+					<Spawn Exe="$(RootDir)/Engine/Restricted/NotForLicensees/Binaries/Win64/unsync.exe" Arguments="hash $(ProjectOutputDirectory)/$(StagedPlatformFolder) -v --buffered-files"/>
+				</Node>
 		
 				<!-- If publishing, we'll c-->
 				<Node Name="$(PlatformPublishNodeName) $(PackageConfiguration)" Requires="$(PlatformPackageNodeName) $(PackageConfiguration)">
@@ -324,17 +322,17 @@
 						This node will be tied in if PublishBuilds=true or DoHashedStagingPublish=true, in whicb case the network path will have been
 						set up in PathSettings.xml  This path will reflect set to scratch or semi-permanent storage based on whether
 						PublishBuilds=true or we're a buld machine (build machines always need to put builds on the network for testing)
-					-->					
+					-->
 					<Do If="$(IsPackaged) and $(PublishBuilds)" >
-						<Property Name="BCRPublishArgs" Value="$(BCRArgs) -skippak -skipstage -skippackage -archive -archivedirectory=&quot;$(NetworkOutputDirectory)/$(StagedPlatformFolder)/$(PackageConfiguration)&quot;" />		
+						<Property Name="BCRPublishArgs" Value="$(BCRArgs) -skippak -skipstage -skippackage -archive -archivedirectory=&quot;$(NetworkOutputDirectory)/$(StagedPlatformFolder)&quot;" />		
 						<Command Name="BuildCookRun" Arguments="$(BCRPublishArgs)" />
 					</Do>
 					<!-- We'll hash-->
 					<Do If="!$(IsPackaged) and ($(PublishBuilds) or $(DoHashedStagingPublish))" >
-						<Copy Files="..." From="$(StagedDir)/$(StagedPlatformFolder)" To="$(NetworkOutputDirectory)/$(StagedPlatformFolder)/Staged"/>
+						<Copy Files="..." From="$(ProjectOutputDirectory)/$(StagedPlatformFolder)" To="$(NetworkOutputDirectory)/$(StagedPlatformFolder)/Staged"/>
 					</Do>
 					<Do If="$(IsPackaged) and $(DoHashedStagingPublish)" >
-						<Copy Files="..." From="$(StagedDir)/$(StagedPlatformFolder)" To="$(NetworkOutputDirectory)/$(StagedPlatformFolder)/Staged"/>
+						<Copy Files="..." From="$(ProjectOutputDirectory)/$(StagedPlatformFolder)" To="$(NetworkOutputDirectory)/$(StagedPlatformFolder)/Staged"/>
 					</Do>
 				</Node>
 				<!-- Tie these nodes into the list for this platform. They will later be tied in (or not) based on overall options -->
@@ -346,17 +344,17 @@
 		</Agent>			
 
 		<!-- Test Time! -->		
-		<Agent Name="Test $(PlatformAgentName)" Type="$(TestAgentType)">		
-			<!-- Prerequisites for the test to run -->
-			<Property Name="TestPrerequisites" Value="$(PlatformEditorNodeName)" />
+		<Agent Name="Test $(PlatformAgentName)" Type="$(TestAgentType)">
 			<!-- Run through all configs -->
 			<ForEach Name="TestConfiguration" Values="$(TestConfigurations)" Separator="+">
+				<!-- Prerequisites for the test to run -->
+				<Property Name="TestPrerequisites" Value="$(PlatformPackageNodeName) $(TestConfiguration)" />
 				<!-- Create the -build argument for Gauntlet based on whether builds are local or on the network-->
-				<Property Name="GauntletBuildArg" Value="-build=&quot;$(ProjectOutputDirectory)/$(TestConfiguration)&quot;" />
+				<Property Name="GauntletBuildArg" Value="-build=&quot;$(ProjectOutputDirectory)&quot;" />
 				<Property Name="GauntletBuildArg" Value="-build=&quot;$(NetworkOutputDirectory)&quot;" If ="'$(NetworkOutputDirectory)' != ''" />
 				<!-- Declare a boot test for the target-->	
 				<Property Name="BootTestName" Value="$(DisplayTargetName) $(TargetPlatform) $(TestConfiguration) BootTest" />
-				<Node Name="$(BootTestName)" After="$(PlatformPublishNodeName) $(TestConfiguration)">
+				<Node Name="$(BootTestName)" After="$(PlatformPublishNodeName) $(TestConfiguration)" Requires="$(TestPrerequisites)">
 					<Do If="'$(AutomationReportOutputDirectory)' != ''">
 						<Property Name="OutputFragmentPath" Value="$(TargetPlatform)/$(TestConfiguration)/UE.BootTest" />
 						<Property Name="OutputReportPath" Value="$(AutomationReportOutputDirectory)/$(OutputFragmentPath)" />
@@ -367,6 +365,7 @@
 					<Command Name="RunUnreal" Arguments="-project=$(TargetProject) -test=&quot;UE.BootTest&quot; -platform=$(TargetPlatform) -configuration=$(TestConfiguration) $(GauntletBuildArg) $(GauntletTargetTestArgs) $(GauntletReportArgs)" />
 				</Node>
 				<!-- Link the boot test in to the prerequisites and generated nodes if its enabled-->
+				<Property Name="TestPrerequisites" Value="$(TestPrerequisites);$(PlatformEditorNodeName)" />
 				<Do If="$(TargetBootTest) and !ContainsItem('$(TargetPlatformsToSkipBootTests)', '$(TargetPlatform)', '+')">
 					<Property Name="TestPrerequisites" Value="$(TestPrerequisites);$(BootTestName)" />
 					<Property Name="AllPlatformBootTestNodes" Value="$(AllPlatformBootTestNodes);$(BootTestName)"/>
diff --git a/Engine/Build/Graph/Tasks/Inc/PathSettings.xml b/Engine/Build/Graph/Tasks/Inc/PathSettings.xml
--- a/Engine/Build/Graph/Tasks/Inc/PathSettings.xml
+++ b/Engine/Build/Graph/Tasks/Inc/PathSettings.xml
@@ -14,6 +14,7 @@
 		InOption: PublishBuilds - Enabled/Disable publishing of builds to build storage
 		
 		InProperty: ProjectName
+		InProperty: ProjectPath
 		InProperty: EscapedBranch
 
 		OutOption: PublishBuilds - Publish output from builds to a network share
@@ -36,8 +37,8 @@
 	<Option Name="PublishReports" DefaultValue="$(IsBuildMachine)" Description="Publish packaged builds to the builds folder on the netwok"/>
 
 
-	<!-- Where builds are saved to locally (todo - use ProjectDir?) -->
-	<Property Name="ProjectOutputDirectory" Value="$(RootDir)/LocalBuilds/$(ProjectName)"/>
+	<!-- Where builds are saved to locally -->
+	<Property Name="ProjectOutputDirectory" Value="$(ProjectPath)/Saved/StagedBuilds"/>
 	<Property Name="AutomationReportOutputDirectory" Value=""/>
 
 	<!-- Scratch network path that can be used for tests if we aren't publishing. Assumes there's a cleanup script.. -->
