Fix 'Update Version Files' defined several times when building from more than one platform

#rnx

diff --git a/Engine/Build/Graph/Tasks/BuildAndTestProject.xml b/Engine/Build/Graph/Tasks/BuildAndTestProject.xml
--- a/Engine/Build/Graph/Tasks/BuildAndTestProject.xml
+++ b/Engine/Build/Graph/Tasks/BuildAndTestProject.xml
@@ -155,11 +155,11 @@
 
 		<Agent Name="Build $(ProjectName)Editor $(EditorPlatform)" Type="$(HostAgentType)">
 			<!-- Update the Version.h and Build.version files if we're making a formal build on a build machine -->
-			<Node Name="Update Version Files">
+			<Node Name="Update Version Files $(EditorPlatform)">
 				<SetVersion Change="$(Change)" Branch="$(EscapedBranch)" If="$(Versioned) and $(IsBuildMachine)"/>
 			</Node>
 			<!-- First generate nodes to build the required tools (e.g. UHT, ShaderCompiler, Editor -->
-			<Node Name="Compile $(ProjectName)Tools $(EditorPlatform)" Requires="Update Version Files">
+			<Node Name="Compile $(ProjectName)Tools $(EditorPlatform)" Requires="Update Version Files $(EditorPlatform)">
 				<Compile Target="UnrealHeaderTool" Platform="$(EditorPlatform)" Configuration="Development"/>
 				<Compile Target="ShaderCompileWorker" Platform="$(EditorPlatform)" Configuration="Development"/>
 				<Compile Target="UnrealPak" Platform="$(EditorPlatform)" Configuration="Development"/>
