Added options for uploading a binary using the online BuildPatchTool

diff --git a/Engine/Build/Graph/Tasks/BuildAndTestProject.xml b/Engine/Build/Graph/Tasks/BuildAndTestProject.xml
--- a/Engine/Build/Graph/Tasks/BuildAndTestProject.xml
+++ b/Engine/Build/Graph/Tasks/BuildAndTestProject.xml
@@ -70,6 +70,13 @@
 	<Option Name="EditorPlatformsToSkipBootTests" DefaultValue="$(EditorPlatformsToSkipTests)" Description="Editor platforms that should skip boot tests."/>
 	<Option Name="TargetPlatformsToSkipBootTests" DefaultValue="$(TargetPlatformsToSkipTests)" Description="Target platforms that should skip boot tests."/>
 
+	<!-- Online BuildPatchTool Settings -->
+	<Option Name="UploadArtifactId" DefaultValue="" Description="ID of the artifact to chunked and uploaded. Upload is ignored if no ID is specified." />
+	<Option Name="UploadPlatform" DefaultValue="" Description="Platform build to upload, i.e. Win64" />
+	<Option Name="UploadConfiguration" DefaultValue="" Description="Configuration to upload, if applicable." />
+	<Option Name="UploadLaunchArgs" DefaultValue="" Description="Arguments to be passed when launching the uploaded app." />	
+	<Option Name="OnlineBPTCredentialsFile" DefaultValue="" Description="Name of the file containing online BuildPatchTool credentials for uploading." />
+
 	<!-- Settings files that can be overriden -->
 	<Option Name="PathSettingsFile" DefaultValue="Inc/PathSettings.xml" Description="Path settings to use" />
 	<Option Name="GauntletSettingsFile" DefaultValue="Inc/GauntletSettings.xml" Description="Gauntlet settings to use" />
@@ -250,6 +257,7 @@
 		<!-- These will be a generated string of all configurations and builds/cooks/packages/tests for this platform that can be tied to aggregates -->
 		<Property Name="AllPlatformBuildNodes" Value="" />
 		<Property Name="AllPlatformCookNodes" Value="" />
+		<Property Name="AllPlatformUploadNodes" Value="" />
 		<Property Name="AllPlatformPackageNodes" Value="" />
 		<Property Name="AllPlatformHashNodes" Value="" />
 
@@ -295,6 +303,7 @@
 
 			<!-- Now package the exe and cooked data-->
 			<ForEach Name="PackageConfiguration" Values="$(PackageConfigurations)" Separator="+">
+				<Property Name="CloudDir" Value="$(ProjectOutputDirectory)/$(StagedPlatformFolder)/CloudDir" />
 
 				<Property Name="IsPackaged" Value="true" />
 				<Property Name="IsPackaged" Value="false" If="ContainsItem('$(PlatformsToSkipPackaging)','$(TargetPlatform)','+')" />	
@@ -315,7 +324,36 @@
 				<Node Name="$(PlatformHashNodeName) $(PackageConfiguration)" Requires="$(PlatformPackageNodeName) $(PackageConfiguration)">
 					<Spawn Exe="$(RootDir)/Engine/Restricted/NotForLicensees/Binaries/Win64/unsync.exe" Arguments="hash $(ProjectOutputDirectory)/$(StagedPlatformFolder) -v --buffered-files"/>
 				</Node>
-		
+				
+				<!-- Upload the build, if applicable. -->
+
+				<Do If="'$(UploadArtifactId)' != '' and '$(UploadConfiguration)' == '$(PackageConfiguration)' and '$(UploadPlatform)' == '$(TargetPlatform)'">
+					<!-- Set UploadPlatform since it isn't always the same as the target -->
+					<Property Name="UploadPlatform" Value="$(TargetPlatform)" />
+					<Property Name="UploadPlatform" Value="Windows" If="'$(TargetPlatform)' == 'Win64'"/>
+					
+					<Property Name="PlatformUploadNodeName" Value="Upload $(TargetName) $(UploadPlatform)" />
+					<Node Name="$(PlatformUploadNodeName)" Requires="$(PlatformPackageNodeName) $(UploadConfiguration)">
+						<Property Name="ConfigurationSuffix" Value="" />	
+						<Property Name="ConfigurationSuffix" Value="-$(TargetPlatform)-$(UploadConfiguration)" If="'$(UploadConfiguration)' != 'Development'" />	
+						<Property Name="Launch" Value="$(ProjectName)/Binaries/$(TargetPlatform)/$(ProjectName)$(ConfigurationSuffix).exe"/>
+						<Tag Files="$(ProjectOutputDirectory)/$(StagedPlatformFolder)/..." Except="*.pdb" With="#MainGameFiles"/>
+					
+						<UploadBinary Platform="$(UploadPlatform)" 
+							ArtifactId="$(ArtifactId)"
+							CommandLineFile="$(OnlineBPTCredentialsFile)"
+							BuildVersion="$(BuildName)" InputDir="$(ProjectOutputDirectory)/$(StagedPlatformFolder)" Files="#MainGameFiles"
+							Launch="$(Launch)" LaunchArgs="$(UploadLaunchArgs)"
+							CloudDir="$(CloudDir)" />
+
+						<LabelBinary ArtifactId="$(ArtifactId)"
+							CommandLineFile="$(OnlineBPTCredentialsFile)"
+							BuildVersion="$(BuildName)" Platform="$(UploadPlatform)" Label="Live" />
+					</Node>	
+
+					<Property Name="AllPlatformUploadNodes" Value="$(AllPlatformUploadNodes);$(PlatformUploadNodeName)" />			
+				</Do>	
+
 				<!-- If publishing, we'll c-->
 				<Node Name="$(PlatformPublishNodeName) $(PackageConfiguration)" Requires="$(PlatformPackageNodeName) $(PackageConfiguration)">
 					<!-- 	
@@ -396,6 +434,7 @@
 		<Label Category="Platforms" Name="$(ProjectName) $(TargetPlatform)" Requires="$(AllPlatformBuildNodes)" />
 		<Label Category="Cook" Name="Cook $(DisplayTargetName) $(TargetPlatform)" Requires="$(AllPlatformCookNodes)" Exclude="$(PlatformEditorNodeName)" />
 		<Label Category="Package" Name="Pkg $(DisplayTargetName) $(TargetPlatform)" Requires="$(AllPlatformPackageNodes);$(AllPlatformHashNodes);$(AllPlatformPublishNodes)" Exclude="$(AllPlatformBuildNodes);$(AllPlatformCookNodes)"/>
+		<Label Category="Upload" Name="Upload $(DisplayTargetName) $(TargetPlatform)" Requires="$(AllPlatformUploadNodes)" />
 		<Label Category="Test" Name="BootTest $(DisplayTargetName) $(TargetPlatform)" Requires="$(AllPlatformBootTestNodes)" Exclude="$(PlatformEditorNodeName)" />
 		<Label Category="Test" Name="Test $(DisplayTargetName) $(TargetPlatform)" Requires="$(AllPlatformTestNodes)" Exclude="$(PlatformEditorNodeName);$(AllPlatformBootTestNodes)"/>
 
@@ -405,6 +444,7 @@
 		<Property Name="PlatformTasks" Value="$(PlatformTasks);$(AllPlatformCookNodes)" If="!$(SkipCook) and !$(SkipPublish)" />		
 		<Property Name="PlatformTasks" Value="$(PlatformTasks);$(AllPlatformPackageNodes);$(AllPlatformPublishNodes)" If="!$(SkipPublish) And $(PublishBuilds)" />	
 		<Property Name="PlatformTasks" Value="$(PlatformTasks);$(AllPlatformHashNodes)"  If="$(PublishHashedStagingData)" />	
+		<Property Name="PlatformTasks" Value="$(PlatformTasks);$(AllPlatformUploadNodes)"  If="'$(AllPlatformUploadNodes)' != ''" />	
 		<Property Name="PlatformTasks" Value="$(PlatformTasks);$(AllPlatformBootTestNodes);$(AllPlatformTestNodes)" />		
 
 		<Property Name="GeneratedTargetNodes" Value="$(GeneratedTargetNodes);$(PlatformTasks)"/>
