Use an empty "complete" tag set to signal cook completion for package nodes so that test agents don't implicitly require all the cooked assets.

#rnx

diff --git a/Engine/Build/Graph/Tasks/BuildAndTestProject.xml b/Engine/Build/Graph/Tasks/BuildAndTestProject.xml
--- a/Engine/Build/Graph/Tasks/BuildAndTestProject.xml
+++ b/Engine/Build/Graph/Tasks/BuildAndTestProject.xml
@@ -291,7 +291,7 @@
 			</ForEach>
 
 			<!-- Cook for the platform. This can run at the same time as the above since we don't need the exe -->
-			<Node Name="$(PlatformCookNodeName)" Requires="$(PlatformEditorNodeName)">
+			<Node Name="$(PlatformCookNodeName)" Requires="$(PlatformEditorNodeName)" Produces="#$(PlatformCookNodeName) Complete">
 				<Property Name="CookArgs" Value="" />				
 				<Property Name="CookArgs" Value="-MapIniSection=$(MapIniSection) $(CookPlatform)" If="'$(MapIniSection)' != ''"/>				
 				<Cook Project="$(TargetProject)" Platform="$(CookPlatform)" Arguments="$(CookArgs)"/>				
@@ -314,7 +314,7 @@
 				<Property Name="BCRArgs" Value="$(BCRArgs) -server" If="'$(DedicatedServer)'" />
 
 				<!-- Pak stage and, if appropriate, package the build using BCR -->
-				<Node Name="$(PlatformPackageNodeName) $(PackageConfiguration)" Requires="$(PlatformCompileNodeName) $(PackageConfiguration);$(PlatformCookNodeName)">
+				<Node Name="$(PlatformPackageNodeName) $(PackageConfiguration)" Requires="$(PlatformCompileNodeName) $(PackageConfiguration);#$(PlatformCookNodeName) Complete">
 						<Property Name="BCRPackageArgs" Value="$(BCRArgs) -pak -stage" />		
 						<Property Name="BCRPackageArgs" Value="$(BCRPackageArgs) -package" If="$(IsPackaged)" />		
 						<Command Name="BuildCookRun" Arguments="$(BCRPackageArgs)" />
