Fix upload step artifact ID parameter

diff --git a/Engine/Build/Graph/Tasks/BuildAndTestProject.xml b/Engine/Build/Graph/Tasks/BuildAndTestProject.xml
--- a/Engine/Build/Graph/Tasks/BuildAndTestProject.xml
+++ b/Engine/Build/Graph/Tasks/BuildAndTestProject.xml
@@ -340,13 +340,13 @@
 						<Tag Files="$(ProjectOutputDirectory)/$(StagedPlatformFolder)/..." Except="*.pdb" With="#MainGameFiles"/>
 					
 						<UploadBinary Platform="$(UploadPlatform)" 
-							ArtifactId="$(ArtifactId)"
+							ArtifactId="$(UploadArtifactId)"
 							CommandLineFile="$(OnlineBPTCredentialsFile)"
 							BuildVersion="$(BuildName)" InputDir="$(ProjectOutputDirectory)/$(StagedPlatformFolder)" Files="#MainGameFiles"
 							Launch="$(Launch)" LaunchArgs="$(UploadLaunchArgs)"
 							CloudDir="$(CloudDir)" />
 
-						<LabelBinary ArtifactId="$(ArtifactId)"
+						<LabelBinary ArtifactId="$(UploadArtifactId)"
 							CommandLineFile="$(OnlineBPTCredentialsFile)"
 							BuildVersion="$(BuildName)" Platform="$(UploadPlatform)" Label="Live" />
 					</Node>	
