Use UnrealEditor target for content only Project.

#rnx

diff --git a/Engine/Build/Graph/Tasks/BuildAndTestProject.xml b/Engine/Build/Graph/Tasks/BuildAndTestProject.xml
--- a/Engine/Build/Graph/Tasks/BuildAndTestProject.xml
+++ b/Engine/Build/Graph/Tasks/BuildAndTestProject.xml
@@ -7,6 +7,8 @@
 	<Option Name="ProjectPath" DefaultValue="None" Restrict=".+" Description="Path to the folder that contains your project"/>
 
 	<Option Name="TargetName" DefaultValue="" Restrict="" Description="If your project supports both Game and Client you can set this, otherwise one will be selected"/>
+	<Option Name="EditorTargetName" DefaultValue="" Restrict="" Description="Specify the editor target name if different than (ProjectName)Editor"/>
+
 
 	<Option Name="ExtraToolCompileArguments" DefaultValue="" Restrict="" Description="Extra arguments to use when building the tools" />
 	<Option Name="ExtraEditorCompileArguments" DefaultValue="" Restrict="" Description="Extra arguments to use when building the editor" />
@@ -110,8 +112,10 @@
 				<Error Message="The project has source but the name the target couldn't be determined. use -set:TargetName:Foo to ensure things build correctly" />
 			</Do>
 			<Property Name="TargetName" Value="UnrealGame" />
+			<Property Name="EditorTargetName" Value="UnrealEditor" If="'$(EditorTargetName)' == ''"/>
 		</Do>
 	</Do>
+	<Property Name="EditorTargetName" Value="$(ProjectName)Editor" If="'$(EditorTargetName)' == ''"/>
 
 	<!-- If skip test is true, add all target platforms to the "do not test" list -->
 	<Do If="$(SkipTest)">
@@ -174,7 +178,7 @@
 				<Compile Target="CrashReportClientEditor" Platform="$(EditorPlatform)" Configuration="Shipping"/>
 			</Node>
 			<Node Name="Compile $(ProjectName)Editor $(EditorPlatform)" Requires="Compile $(ProjectName)Tools $(EditorPlatform)">
-				<Compile Target="$(ProjectName)Editor" Platform="$(EditorPlatform)" Configuration="Development" Arguments="$(ExtraEditorCompileArguments)"/>
+				<Compile Target="$(EditorTargetName)" Platform="$(EditorPlatform)" Configuration="Development" Arguments="$(ExtraEditorCompileArguments)"/>
 			</Node>
 			<!--Add these nodes to our dependency list -->
 			<Property Name="GeneratedToolNodes" Value="Compile $(ProjectName)Editor $(EditorPlatform); Compile $(ProjectName)Tools $(EditorPlatform)" If="!$(SkipBuild)" />
