Added option for appending obfuscation launch args when uploading to EGS

diff --git a/Engine/Build/Graph/Tasks/BuildAndTestProject.xml b/Engine/Build/Graph/Tasks/BuildAndTestProject.xml
--- a/Engine/Build/Graph/Tasks/BuildAndTestProject.xml
+++ b/Engine/Build/Graph/Tasks/BuildAndTestProject.xml
@@ -77,6 +77,7 @@
 	<Option Name="UploadPlatform" DefaultValue="" Description="Platform build to upload, i.e. Win64" />
 	<Option Name="UploadConfiguration" DefaultValue="" Description="Configuration to upload, if applicable." />
 	<Option Name="UploadLaunchArgs" DefaultValue="" Description="Arguments to be passed when launching the uploaded app." />	
+	<Option Name="UploadAppendObfuscationLaunchArg" Restrict="true|false" DefaultValue="true" Description="Whether to append an random obfuscation id launch arg." />	
 	<Option Name="OnlineBPTCredentialsFile" DefaultValue="" Description="Name of the file containing online BuildPatchTool credentials for uploading." />
 
 	<!-- Settings files that can be overriden -->
@@ -347,7 +348,7 @@
 							ArtifactId="$(UploadArtifactId)"
 							CommandLineFile="$(OnlineBPTCredentialsFile)"
 							BuildVersion="$(BuildName)" InputDir="$(ProjectOutputDirectory)/$(StagedPlatformFolder)" Files="#MainGameFiles"
-							Launch="$(Launch)" LaunchArgs="$(UploadLaunchArgs)"
+							Launch="$(Launch)" LaunchArgs="$(UploadLaunchArgs)" AppendObfuscationLaunchArg="$(UploadAppendObfuscationLaunchArg)"
 							CloudDir="$(CloudDir)" />
 
 						<LabelBinary ArtifactId="$(UploadArtifactId)"
