Only default target configurations to specify Development when there were no target, package, or test configurations specified

diff --git a/Engine/Build/Graph/Tasks/Inc/PlatformOptions.xml b/Engine/Build/Graph/Tasks/Inc/PlatformOptions.xml
--- a/Engine/Build/Graph/Tasks/Inc/PlatformOptions.xml
+++ b/Engine/Build/Graph/Tasks/Inc/PlatformOptions.xml
@@ -182,7 +182,6 @@
 
 <Do If="'$(TargetConfigurations)' == ''" >
 	<Property Name="TargetConfigurations" Value="$(DefaultTargetConfigurations)" />
-	<Property Name="TargetConfigurations" Value="Development" If="'$(TargetConfigurations)' == ''" />
 </Do>
 
 <Do If="'$(PackageConfigurations)' == ''" >	
@@ -209,4 +208,7 @@
 	</Do>
 </ForEach>
 
+<!-- If there were no configurations specified, then default to Development for all targets -->
+<Property Name="TargetConfigurations" Value="Development" If="'$(TargetConfigurations)' == ''" />
+
 </BuildGraph>
\ No newline at end of file
