Moved upload cloud directory so it is not under the BuildRoot directory

diff --git a/Engine/Build/Graph/Tasks/BuildAndTestProject.xml b/Engine/Build/Graph/Tasks/BuildAndTestProject.xml
--- a/Engine/Build/Graph/Tasks/BuildAndTestProject.xml
+++ b/Engine/Build/Graph/Tasks/BuildAndTestProject.xml
@@ -308,7 +308,7 @@
 
 			<!-- Now package the exe and cooked data-->
 			<ForEach Name="PackageConfiguration" Values="$(PackageConfigurations)" Separator="+">
-				<Property Name="CloudDir" Value="$(ProjectOutputDirectory)/$(StagedPlatformFolder)/CloudDir" />
+				<Property Name="CloudDir" Value="$(ProjectOutputDirectory)/CloudDir/$(StagedPlatformFolder)" />
 
 				<Property Name="IsPackaged" Value="true" />
 				<Property Name="IsPackaged" Value="false" If="ContainsItem('$(PlatformsToSkipPackaging)','$(TargetPlatform)','+')" />	
