Defaults Package and Test configurations to whatever configuration was set for Target

diff --git a/Engine/Build/Graph/Tasks/Inc/PlatformOptions.xml b/Engine/Build/Graph/Tasks/Inc/PlatformOptions.xml
--- a/Engine/Build/Graph/Tasks/Inc/PlatformOptions.xml
+++ b/Engine/Build/Graph/Tasks/Inc/PlatformOptions.xml
@@ -208,7 +208,9 @@
 	</Do>
 </ForEach>
 
-<!-- If there were no configurations specified, then default to Development for all targets -->
+<!-- If no target configurations were added, then no configurations were specified. Default to Development. -->
 <Property Name="TargetConfigurations" Value="Development" If="'$(TargetConfigurations)' == ''" />
+<Property Name="PackageConfigurations" Value="$(TargetConfigurations)" If="'$(PackageConfigurations)' == ''" />
+<Property Name="TestConfigurations" Value="$(PackageConfigurations)" If="'$(TestConfigurations)' == ''" />
 
 </BuildGraph>
\ No newline at end of file
