Fixed editor boot test being enabled even when EditorBootTest was set to false

diff --git a/Engine/Build/Graph/Tasks/BuildAndTestProject.xml b/Engine/Build/Graph/Tasks/BuildAndTestProject.xml
--- a/Engine/Build/Graph/Tasks/BuildAndTestProject.xml
+++ b/Engine/Build/Graph/Tasks/BuildAndTestProject.xml
@@ -217,7 +217,7 @@
 			<Do If="ContainsItem('$(EditorPlatforms)', '$(EditorPlatform)', '+') and !ContainsItem('$(EditorPlatformsToSkipTests)', '$(EditorPlatform)', '+')">
 				<Label Category="Test" Name="BootTest $(ProjectName)Editor $(EditorPlatform)" Requires="$(PlatformBootTestNodeName)" />
 				<Label Category="Test" Name="Test $(ProjectName)Editor $(EditorPlatform)" Requires="$(PlatformEditorTestNodes)" />
-				<Property Name="GeneratedEditorTestNodes" Value="$(GeneratedEditorTestNodes);$(PlatformBootTestNodeName);$(PlatformEditorTestNodes)" />
+				<Property Name="GeneratedEditorTestNodes" Value="$(GeneratedEditorTestNodes);$(PlatformEditorTestNodes)" />
 			</Do>
 		</Agent>
 	</ForEach>
