Simplified BuildAndTestProject script and fixed use of unsync hashing.
 - Added an explicit "Stage" step. The results of this step are hashed and published before packaging is done.
 - Staging, hashing, and packaging is done as single steps, rather than one step per config. This avoids redundant re-staging and stomping previously built data.
 - Merged target compile steps into a single, multi-config step.

diff --git a/Engine/Build/Graph/Tasks/BuildAndTestProject.xml b/Engine/Build/Graph/Tasks/BuildAndTestProject.xml
--- a/Engine/Build/Graph/Tasks/BuildAndTestProject.xml
+++ b/Engine/Build/Graph/Tasks/BuildAndTestProject.xml
@@ -90,8 +90,7 @@
 	<Property Name="TargetProject" Value="$(ProjectPath)/$(ProjectName).uproject" />
 	<Error Message="Project file $(TargetProject) does not exist" If="!Exists('$(TargetProject)')" />
 
-
-	<!-- If the user didn't supply a TargetName then try to auto detect is. This is reliable uness the project has 'Game' in it but the target doesn't -->
+	<!-- If the user didn't supply a TargetName then try to auto detect it. This is reliable uness the project has 'Game' in it but the target doesn't -->
 	<Do If="'$(TargetName)' == ''">
 		<!-- 
 			If the user didn't supply a TargetName, try to figure it out based on GameName.Target or GameNameClient.target.
@@ -222,7 +221,7 @@
 		</Agent>
 	</ForEach>
 
-	<!-- Now, for each platform we build, cook, package, and test -->
+	<!-- Now, for each platform we build, cook, stage, package, publish, and test -->
 	<ForEach Name="TargetPlatform" Values="$(TargetPlatforms)" Separator="+">
 
 		<!-- Declare type of host required to build & cook this platform-->
@@ -247,19 +246,11 @@
 		<Property Name="PlatformAgentName" Value="$(DisplayTargetName) $(TargetPlatform)" />
 		<Property Name="PlatformCompileNodeName" Value="Compile $(TargetName) $(TargetPlatform)" />
 		<Property Name="PlatformCookNodeName" Value="Cook $(DisplayTargetName) $(TargetPlatform)" />
+		<Property Name="PlatformStageNodeName" Value="Stage $(DisplayTargetName) $(TargetPlatform)" />
 		<Property Name="PlatformPackageNodeName" Value="Package $(DisplayTargetName) $(TargetPlatform)" />
 		<Property Name="PlatformHashNodeName" Value="Hash $(DisplayTargetName) $(TargetPlatform)" />
-		<Property Name="PlatformPublishNodeName" Value="Publish $(DisplayTargetName) $(TargetPlatform)" />
-
-		<!-- These will be a generated string of all configurations and builds/cooks/packages/tests for this platform that can be tied to aggregates -->
-		<Property Name="AllPlatformBuildNodes" Value="" />
-		<Property Name="AllPlatformCookNodes" Value="" />
-		<Property Name="AllPlatformPackageNodes" Value="" />
-		<Property Name="AllPlatformHashNodes" Value="" />
-
-		<Property Name="AllPlatformPublishNodes" Value="" />
-		<Property Name="AllPlatformBootTestNodes" Value="" />
-		<Property Name="AllPlatformTestNodes" Value="" />
+		<Property Name="PlatformPublishStagedNodeName" Value="Publish Staged $(DisplayTargetName) $(TargetPlatform)" />
+		<Property Name="PlatformPublishPackagedNodeName" Value="Publish Packaged $(DisplayTargetName) $(TargetPlatform)" />
 		
 		<!-- Set CookPlatform since it isn't always the same as the target -->
 		<Property Name="CookPlatform" Value="$(TargetPlatform)" />
@@ -267,97 +258,83 @@
 		<Property Name="CookPlatform" Value="$(CookPlatform)_ETC2" If="'$(TargetPlatform)' == 'Android'"/>
 		<Property Name="CookPlatform" Value="$(CookPlatform)Client" If="'$(TargetType)' == 'Client'" />
 
-		<Property Name="BCRPlatform" Value="$(TargetPlatform)" />
-
 		<Property Name="BCRCookFlavor" Value="" If="" />
 		<Property Name="BCRCookFlavor" Value="ETC2" If="'$(TargetPlatform)' == 'Android'" />
+		
 		<Property Name="StagedPlatformFolder" Value="$(CookPlatform)"  />
 		<!-- <Property Name="StagedPlatformFolder" Value="$(StagedPlatformFolder)Client" If="'$(TargetType)' == 'Client'" /> -->
 
+		<!-- Hashing of staged data can only be done on Win64 hosts -->
+		<Property Name="DoUnsyncHashing" Value="$(HostAgentType) == 'Win64'" />
+		
 		<!-- Build the client and if necessary a server-->
 		<Agent Name="Make $(PlatformAgentName)" Type="$(BuildAgentType)">
-			<ForEach Name="TargetConfiguration" Values="$(TargetConfigurations)" Separator="+">
-				<Node Name="$(PlatformCompileNodeName) $(TargetConfiguration)" Requires="$(PlatformToolsNodeName)" RunEarly="true">
+
+			<!-- Default BuildCookRun args needed when staging/packaging/archiving this configuration -->
+			<Property Name="BCRArgs" Value="-project=$(TargetProject) -platform=$(TargetPlatform) -cookflavor=&quot;$(BCRCookFlavor)&quot;" />
+			<Property Name="BCRArgs" Value="$(BCRArgs) -client" If="'$(TargetType)' == 'Client'" />
+			<Property Name="BCRArgs" Value="$(BCRArgs) -server" If="'$(DedicatedServer)'" />
+			
+			<!-- Compile -->
+			<Node Name="$(PlatformCompileNodeName)" Requires="$(PlatformToolsNodeName)" RunEarly="true">
+				<ForEach Name="TargetConfiguration" Values="$(TargetConfigurations)" Separator="+">
 					<Compile Target="$(TargetName)" Platform="$(TargetPlatform)" Configuration="$(TargetConfiguration)" Arguments="$(ExtraTargetCompileArguments)"/>
 					<Do If="'$(DedicatedServer)'" >
 						<Compile Target="$(TargetName)Server" Platform="Win64" Configuration="$(TargetConfiguration)" Arguments="$(ExtraTargetCompileArguments)"/>
-					</Do>				
-				</Node>
-				<Property Name="AllPlatformBuildNodes" Value="$(AllPlatformBuildNodes);$(PlatformCompileNodeName) $(TargetConfiguration)" />
-			</ForEach>
+					</Do>
+				</ForEach>
+			</Node>
 
-			<!-- Cook for the platform. This can run at the same time as the above since we don't need the exe -->
+			<!-- Cook -->
 			<Node Name="$(PlatformCookNodeName)" Requires="$(PlatformEditorNodeName)" Produces="#$(PlatformCookNodeName) Complete">
-				<Property Name="CookArgs" Value="" />				
+				<Property Name="CookArgs" Value="" />
 				<Property Name="CookArgs" Value="-MapIniSection=$(MapIniSection) $(CookPlatform)" If="'$(MapIniSection)' != ''"/>				
-				<Cook Project="$(TargetProject)" Platform="$(CookPlatform)" Arguments="$(CookArgs)"/>				
-			</Node>			
-			<Property Name="AllPlatformCookNodes" Value="$(AllPlatformCookNodes);$(PlatformCookNodeName)" />
-
-			<!-- Do we publish a hashed version of the staged data? We only want to do this once per platform so use a local var that we'll flip to false later -->
-			<Property Name="DoHashedStagingPublish" Value="$(HostAgentType) == 'Win64'" />
-
-			<!-- Now package the exe and cooked data-->
-			<ForEach Name="PackageConfiguration" Values="$(PackageConfigurations)" Separator="+">
-				<Property Name="IsPackaged" Value="true" />
-				<Property Name="IsPackaged" Value="false" If="ContainsItem('$(PlatformsToSkipPackaging)','$(TargetPlatform)','+')" />	
-
-				<!-- Default BuildCookRun args needed staging/packaging/archiving this configuration -->
-				<Property Name="BCRArgs" Value="-project=$(TargetProject) -platform=$(BCRPlatform) -cookflavor=&quot;$(BCRCookFlavor)&quot; -configuration=$(PackageConfiguration) $(ExtraStageAndPackageArguments) -skipbuild -skipcook" />
-				<Property Name="BCRArgs" Value="$(BCRArgs) -client" If="'$(TargetType)' == 'Client'" />
-				<Property Name="BCRArgs" Value="$(BCRArgs) -server" If="'$(DedicatedServer)'" />
-
-				<!-- Pak stage and, if appropriate, package the build using BCR -->
-				<Node Name="$(PlatformPackageNodeName) $(PackageConfiguration)" Requires="$(PlatformCompileNodeName) $(PackageConfiguration);#$(PlatformCookNodeName) Complete">
-						<Property Name="BCRPackageArgs" Value="$(BCRArgs) -pak -stage" />		
-						<Property Name="BCRPackageArgs" Value="$(BCRPackageArgs) -package" If="$(IsPackaged)" />		
-						<Command Name="BuildCookRun" Arguments="$(BCRPackageArgs)" />
-				</Node>
-
-				<!-- hash the content (we'll only tie one of these into our nodes, but always declare them as per best practices)-->
-				<Node Name="$(PlatformHashNodeName) $(PackageConfiguration)" Requires="$(PlatformPackageNodeName) $(PackageConfiguration)">
+				<Cook Project="$(TargetProject)" Platform="$(CookPlatform)" Arguments="$(CookArgs) $(ExtraCookArguments)"/>
+			</Node>
+			
+			<!-- Stage -->
+			<Node Name="$(PlatformStageNodeName)" Requires="$(PlatformCompileNodeName);#$(PlatformCookNodeName) Complete">
+				<Command Name="BuildCookRun" Arguments="$(BCRArgs) $(ExtraStageAndPackageArguments) -configuration=$(PackageConfigurations) -skipbuild -skipcook -stage -pak" />
+			</Node>
+			
+			<!-- Hash -->
+			<Node Name="$(PlatformHashNodeName)" Requires="$(PlatformStageNodeName)">
+				<Do If="$(DoUnsyncHashing)">
 					<Spawn Exe="$(RootDir)/Engine/Restricted/NotForLicensees/Binaries/Win64/unsync.exe" Arguments="hash $(ProjectOutputDirectory)/$(StagedPlatformFolder) -v --buffered-files"/>
-				</Node>
-
-				<!-- If publishing, we'll c-->
-				<Node Name="$(PlatformPublishNodeName) $(PackageConfiguration)" Requires="$(PlatformPackageNodeName) $(PackageConfiguration)">
-					<!-- 	
-						This node will be tied in if PublishBuilds=true or DoHashedStagingPublish=true, in whicb case the network path will have been
-						set up in PathSettings.xml  This path will reflect set to scratch or semi-permanent storage based on whether
-						PublishBuilds=true or we're a buld machine (build machines always need to put builds on the network for testing)
-					-->
-					<Do If="$(IsPackaged) and $(PublishBuilds)" >
-						<Property Name="BCRPublishArgs" Value="$(BCRArgs) -skippak -skipstage -skippackage -archive -archivedirectory=&quot;$(NetworkOutputDirectory)/$(StagedPlatformFolder)&quot;" />		
-						<Command Name="BuildCookRun" Arguments="$(BCRPublishArgs)" />
-					</Do>
-					<!-- We'll hash-->
-					<Do If="!$(IsPackaged) and ($(PublishBuilds) or $(DoHashedStagingPublish))" >
-						<Copy Files="..." From="$(ProjectOutputDirectory)/$(StagedPlatformFolder)" To="$(NetworkOutputDirectory)/$(StagedPlatformFolder)/Staged"/>
-					</Do>
-					<Do If="$(IsPackaged) and $(DoHashedStagingPublish)" >
-						<Copy Files="..." From="$(ProjectOutputDirectory)/$(StagedPlatformFolder)" To="$(NetworkOutputDirectory)/$(StagedPlatformFolder)/Staged"/>
-					</Do>
-				</Node>
-				<!-- Tie these nodes into the list for this platform. They will later be tied in (or not) based on overall options -->
-				<Property Name="AllPlatformPackageNodes" Value="$(AllPlatformPackageNodes);$(PlatformPackageNodeName) $(PackageConfiguration)" />
-				<Property Name="AllPlatformPublishNodes" Value="$(AllPlatformPublishNodes);$(PlatformPublishNodeName) $(PackageConfiguration)" />
-				<Property Name="AllPlatformHashNodes" Value="$(AllPlatformHashNodes);$(PlatformHashNodeName) $(PackageConfiguration)" If="$(DoHashedStagingPublish)" />
-				<Property Name="DoHashedStagingPublish" Value="false" />
-			</ForEach>
-		</Agent>			
+				</Do>
+			</Node>
+			
+			<!-- Publish (Staged) -->
+			<Node Name="$(PlatformPublishStagedNodeName)" Requires="$(PlatformHashNodeName)">
+				<Copy Files="..." From="$(ProjectOutputDirectory)/$(StagedPlatformFolder)" To="$(NetworkOutputDirectory)/$(StagedPlatformFolder)/Staged"/>
+			</Node>
+			
+			<!-- Package -->
+			<Node Name="$(PlatformPackageNodeName)" Requires="$(PlatformPublishStagedNodeName)">
+				<Command Name="BuildCookRun" Arguments="$(BCRArgs) $(ExtraStageAndPackageArguments) -configuration=$(PackageConfigurations) -skipbuild -skipcook -skipstage -skippak -package" />
+			</Node>
+			
+			<!-- Publish (Packages) -->
+			<Node Name="$(PlatformPublishPackagedNodeName)" Requires="$(PlatformPackageNodeName)">
+				<Command Name="BuildCookRun" Arguments="$(BCRArgs) $(ExtraStageAndPackageArguments) -configuration=$(PackageConfigurations) -skipbuild -skipcook -skipstage -skippak -skippackage -archive -archivedirectory=&quot;$(NetworkOutputDirectory)/$(StagedPlatformFolder)/Packages&quot;" />
+			</Node>
+			
+		</Agent>
 
-		<!-- Test Time! -->		
+		<!-- Test Time! -->
+		<Property Name="AllPlatformBootTestNodes" Value="" />
+		<Property Name="AllPlatformTestNodes" Value="" />
 		<Agent Name="Test $(PlatformAgentName)" Type="$(TestAgentType)">
 			<!-- Run through all configs -->
 			<ForEach Name="TestConfiguration" Values="$(TestConfigurations)" Separator="+">
 				<!-- Prerequisites for the test to run -->
-				<Property Name="TestPrerequisites" Value="$(PlatformPackageNodeName) $(TestConfiguration)" />
+				<Property Name="TestPrerequisites" Value="$(PlatformPackageNodeName)" />
 				<!-- Create the -build argument for Gauntlet based on whether builds are local or on the network-->
 				<Property Name="GauntletBuildArg" Value="-build=&quot;$(ProjectOutputDirectory)&quot;" />
 				<Property Name="GauntletBuildArg" Value="-build=&quot;$(NetworkOutputDirectory)&quot;" If ="'$(NetworkOutputDirectory)' != ''" />
 				<!-- Declare a boot test for the target-->	
 				<Property Name="BootTestName" Value="$(DisplayTargetName) $(TargetPlatform) $(TestConfiguration) BootTest" />
-				<Node Name="$(BootTestName)" After="$(PlatformPublishNodeName) $(TestConfiguration)" Requires="$(TestPrerequisites)">
+				<Node Name="$(BootTestName)" After="$(PlatformPublishPackagedNodeName)" Requires="$(TestPrerequisites)">
 					<Do If="'$(AutomationReportOutputDirectory)' != ''">
 						<Property Name="OutputFragmentPath" Value="$(TargetPlatform)/$(TestConfiguration)/UE.BootTest" />
 						<Property Name="OutputReportPath" Value="$(AutomationReportOutputDirectory)/$(OutputFragmentPath)" />
@@ -383,7 +360,7 @@
 						<Property Name="GauntletReportArgs" Value="$(GauntletReportArgs) -ReportURL=&quot;$(AutomationReportUrl)/$(OutputFragmentPath)/&quot;" If="'$(AutomationReportUrl)' != ''" />
 					</Do>
 					<Property Name="TestNodeName" Value="$(DisplayTargetName) $(TargetPlatform) $(TestConfiguration) $(TestName)" />
-					<!-- Run the test. No 'Requires' as we run on the same machine locally and use the publish folder otherwose -->
+					<!-- Run the test. No 'Requires' as we run on the same machine locally and use the publish folder otherwise -->
 					<Node Name="$(TestNodeName)" Requires="$(TestPrerequisites)">
 						<Command Name="RunUnreal" Arguments="-project=$(TargetProject) -test=&quot;$(TestName)&quot; -platform=$(TargetPlatform) -configuration=$(TestConfiguration) $(GauntletBuildArg) $(GauntletTargetTestArgs) $(GauntletReportArgs)"/>
 					</Node>
@@ -396,18 +373,19 @@
 		</Agent>
 
 		<!-- Declare labels for CIS -->
-		<Label Category="Platforms" Name="$(ProjectName) $(TargetPlatform)" Requires="$(AllPlatformBuildNodes)" />
-		<Label Category="Cook" Name="Cook $(DisplayTargetName) $(TargetPlatform)" Requires="$(AllPlatformCookNodes)" Exclude="$(PlatformEditorNodeName)" />
-		<Label Category="Package" Name="Pkg $(DisplayTargetName) $(TargetPlatform)" Requires="$(AllPlatformPackageNodes);$(AllPlatformHashNodes);$(AllPlatformPublishNodes)" Exclude="$(AllPlatformBuildNodes);$(AllPlatformCookNodes)"/>
+		<Label Category="Platforms" Name="$(ProjectName) $(TargetPlatform)" Requires="$(PlatformCompileNodeName)" />
+		<Label Category="Cook" Name="Cook $(DisplayTargetName) $(TargetPlatform)" Requires="$(PlatformCookNodeName)" />
+		<Label Category="Stage" Name="Stage $(DisplayTargetName) $(TargetPlatform)" Requires="$(PlatformStageNodeName);$(PlatformHashNodeName)" />
+		<Label Category="Package" Name="Pkg $(DisplayTargetName) $(TargetPlatform)" Requires="$(PlatformPackageNodeName);$(PlatformPublishPackagedNodeName)"/>
 		<Label Category="Test" Name="BootTest $(DisplayTargetName) $(TargetPlatform)" Requires="$(AllPlatformBootTestNodes)" Exclude="$(PlatformEditorNodeName)" />
 		<Label Category="Test" Name="Test $(DisplayTargetName) $(TargetPlatform)" Requires="$(AllPlatformTestNodes)" Exclude="$(PlatformEditorNodeName);$(AllPlatformBootTestNodes)"/>
 
-		<!-- Tie all steps into our platform tasks uless excluded -->
+		<!-- Tie all steps into our platform tasks unless excluded -->
 		<Property Name="PlatformTasks" Value="" />
-		<Property Name="PlatformTasks" Value="$(AllPlatformBuildNodes)" If="!$(SkipBuild)" />
-		<Property Name="PlatformTasks" Value="$(PlatformTasks);$(AllPlatformCookNodes)" If="!$(SkipCook) and !$(SkipPublish)" />		
-		<Property Name="PlatformTasks" Value="$(PlatformTasks);$(AllPlatformPackageNodes);$(AllPlatformPublishNodes)" If="!$(SkipPublish) And $(PublishBuilds)" />	
-		<Property Name="PlatformTasks" Value="$(PlatformTasks);$(AllPlatformHashNodes)"  If="$(PublishHashedStagingData)" />	
+		<Property Name="PlatformTasks" Value="$(PlatformCompileNodeName)" If="!$(SkipBuild)" />
+		<Property Name="PlatformTasks" Value="$(PlatformTasks);$(PlatformCookNodeName)" If="!$(SkipCook) and !$(SkipPublish)" />		
+		<Property Name="PlatformTasks" Value="$(PlatformTasks);$(PlatformPackageNodeName);$(PlatformPublishPackagedNodeName)" If="!$(SkipPublish) And $(PublishBuilds)" />	
+		<Property Name="PlatformTasks" Value="$(PlatformTasks);$(PlatformHashNodeName)"  If="$(PublishHashedStagingData)" />	
 		<Property Name="PlatformTasks" Value="$(PlatformTasks);$(AllPlatformBootTestNodes);$(AllPlatformTestNodes)" />		
 
 		<Property Name="GeneratedTargetNodes" Value="$(GeneratedTargetNodes);$(PlatformTasks)"/>
