Add symbol upload support
 - Enabled when SymbolServerPath is defined in the parent project script
 - Also defines a "EnableSymbolStore" option, which defaults to true on build machines.

diff --git a/Engine/Build/Graph/Tasks/BuildAndTestProject.xml b/Engine/Build/Graph/Tasks/BuildAndTestProject.xml
--- a/Engine/Build/Graph/Tasks/BuildAndTestProject.xml
+++ b/Engine/Build/Graph/Tasks/BuildAndTestProject.xml
@@ -39,8 +39,12 @@
 		<Property Name="NetworkTempRootOverride" Value="" />
 		<Property Name="NetworkPublishRootOverride" Value="" />
 		<Property Name="NetworkReportRootOverride" Value="" />
+		<Property Name="SymbolServerPath" Value="" />
 	</Do>
 
+	<Option Name="EnableSymbolStore" DefaultValue="$(IsBuildMachine)" Restrict="true|false" Description="Enables the upload of debugging symbols" />
+	<Property Name="SymbolServerPath" Value="" If="$(EnableSymbolStore) == 'false'"/>
+
 	<!-- This will define options for all platforms and create lists called EditorPlatforms, RequiredEditorPlatforms, and TargetPlatforms. E.g. TargetPlatforms=Win64+Mac+PS4 -->
 	<Include Script="Inc/PlatformOptions.xml" />
 
@@ -278,9 +282,17 @@
 			<!-- Compile -->
 			<Node Name="$(PlatformCompileNodeName)" Requires="$(PlatformToolsNodeName)" RunEarly="true">
 				<ForEach Name="TargetConfiguration" Values="$(TargetConfigurations)" Separator="+">
-					<Compile Target="$(TargetName)" Platform="$(TargetPlatform)" Configuration="$(TargetConfiguration)" Arguments="$(ExtraTargetCompileArguments)"/>
+					
+					<!-- Build the binaries -->
+					<Compile Target="$(TargetName)" Platform="$(TargetPlatform)" Configuration="$(TargetConfiguration)" Arguments="$(ExtraTargetCompileArguments)" Tag="#BinaryFiles_$(TargetPlatform)_$(TargetConfiguration)"/>
 					<Do If="'$(DedicatedServer)'" >
-						<Compile Target="$(TargetName)Server" Platform="Win64" Configuration="$(TargetConfiguration)" Arguments="$(ExtraTargetCompileArguments)"/>
+						<Compile Target="$(TargetName)Server" Platform="Win64" Configuration="$(TargetConfiguration)" Arguments="$(ExtraTargetCompileArguments)" Tag="#BinaryFiles_$(TargetPlatform)_$(TargetConfiguration)"/>
+					</Do>
+					
+					<!-- Upload binaries to symbol store, if enabled -->
+					<Do If="'$(SymbolServerPath)' != ''">
+						<Log Message="Uploading symbols to &quot;$(SymbolServerPath)&quot; ..."/>
+						<SymStore Platform="$(TargetPlatform)" Files="#BinaryFiles_$(TargetPlatform)_$(TargetConfiguration)" StoreDir="$(SymbolServerPath)/$(TargetPlatform)" Product="$(TargetName)" BuildVersion="$(BuildName)" />
 					</Do>
 				</ForEach>
 			</Node>
@@ -300,7 +312,7 @@
 			<!-- Hash -->
 			<Node Name="$(PlatformHashNodeName)" Requires="$(PlatformStageNodeName)">
 				<Do If="$(DoUnsyncHashing)">
-					<Spawn Exe="$(RootDir)/Engine/Restricted/NotForLicensees/Binaries/Win64/unsync.exe" Arguments="hash $(ProjectOutputDirectory)/$(StagedPlatformFolder) -v --buffered-files"/>
+					<Spawn Exe="$(RootDir)/Engine/Restricted/NotForLicensees/Binaries/Win64/unsync.exe" Arguments="hash &quot;$(ProjectOutputDirectory)/$(StagedPlatformFolder)&quot; -v --buffered-files"/>
 				</Do>
 			</Node>
 			
