Enable editor boot tests to be run alone not just as prerequisites for other tests

diff --git a/Engine/Build/Graph/Tasks/BuildAndTestProject.xml b/Engine/Build/Graph/Tasks/BuildAndTestProject.xml
--- a/Engine/Build/Graph/Tasks/BuildAndTestProject.xml
+++ b/Engine/Build/Graph/Tasks/BuildAndTestProject.xml
@@ -210,7 +210,7 @@
 					<Command Name="RunUnreal" Arguments="-project=$(TargetProject) -build=Editor -test=&quot;$(TestName)&quot; $(GauntletReportArgs) $(GauntletEditorTestArgs)" />
 				</Node>
 				<!-- Add each test node to our generated list -->
-				<Property Name="PlatformEditorTestNodes" Value="$(PlatformEditorTestNodes);$(TestNodeName)"  />
+				<Property Name="PlatformEditorTestNodes" Value="$(PlatformEditorTestNodes);$(TestNodeName)" />
 			</ForEach>
 
 			<!-- Create Labels -->
@@ -220,6 +220,7 @@
 			<Do If="ContainsItem('$(EditorPlatforms)', '$(EditorPlatform)', '+') and !ContainsItem('$(EditorPlatformsToSkipTests)', '$(EditorPlatform)', '+')">
 				<Label Category="Test" Name="BootTest $(ProjectName)Editor $(EditorPlatform)" Requires="$(PlatformBootTestNodeName)" />
 				<Label Category="Test" Name="Test $(ProjectName)Editor $(EditorPlatform)" Requires="$(PlatformEditorTestNodes)" />
+				<Property Name="GeneratedEditorTestNodes" Value="$(GeneratedEditorTestNodes);$(PlatformBootTestNodeName)" If="$(EditorBootTest)" />
 				<Property Name="GeneratedEditorTestNodes" Value="$(GeneratedEditorTestNodes);$(PlatformEditorTestNodes)" />
 			</Do>
 		</Agent>
