Temporary fix for gauntlet tests when there are multiple build configurations by specifying the configuration subfolder as the test build directory when applicable

diff --git a/Engine/Build/Graph/Tasks/BuildAndTestProject.xml b/Engine/Build/Graph/Tasks/BuildAndTestProject.xml
--- a/Engine/Build/Graph/Tasks/BuildAndTestProject.xml
+++ b/Engine/Build/Graph/Tasks/BuildAndTestProject.xml
@@ -342,9 +342,18 @@
 			<ForEach Name="TestConfiguration" Values="$(TestConfigurations)" Separator="+">
 				<!-- Prerequisites for the test to run -->
 				<Property Name="TestPrerequisites" Value="$(PlatformPackageNodeName)" />
+				
 				<!-- Create the -build argument for Gauntlet based on whether builds are local or on the network-->
 				<Property Name="GauntletBuildArg" Value="-build=&quot;$(ProjectOutputDirectory)&quot;" />
-				<Property Name="GauntletBuildArg" Value="-build=&quot;$(NetworkOutputDirectory)&quot;" If ="'$(NetworkOutputDirectory)' != ''" />
+
+				<Do If="'$(NetworkOutputDirectory)' != ''">
+					<Property Name="TestDirectory" Value="$(NetworkOutputDirectory)" />
+					<Do If="Exists('$(TestDirectory)/$(StagedPlatformFolder)/Packages/$(TestConfiguration)')">
+						<Property Name="TestDirectory" Value="$(TestDirectory)/$(StagedPlatformFolder)/Packages/$(TestConfiguration)" />
+					</Do>
+					<Property Name="GauntletBuildArg" Value="-build=&quot;$(TestDirectory)&quot;" />
+				</Do>
+				
 				<!-- Declare a boot test for the target-->	
 				<Property Name="BootTestName" Value="$(DisplayTargetName) $(TargetPlatform) $(TestConfiguration) BootTest" />
 				<Node Name="$(BootTestName)" After="$(PlatformPublishPackagedNodeName)" Requires="$(TestPrerequisites)">
