Adding nocompile flag to the BuildAndTestProject.xml.

diff --git a/Engine/Build/Graph/Tasks/BuildAndTestProject.xml b/Engine/Build/Graph/Tasks/BuildAndTestProject.xml
--- a/Engine/Build/Graph/Tasks/BuildAndTestProject.xml
+++ b/Engine/Build/Graph/Tasks/BuildAndTestProject.xml
@@ -276,7 +276,7 @@
 		<Agent Name="Make $(PlatformAgentName)" Type="$(BuildAgentType)">
 
 			<!-- Default BuildCookRun args needed when staging/packaging/archiving this configuration -->
-			<Property Name="BCRArgs" Value="-project=$(TargetProject) -platform=$(TargetPlatform) -cookflavor=&quot;$(BCRCookFlavor)&quot;" />
+			<Property Name="BCRArgs" Value="-nocompile -project=$(TargetProject) -platform=$(TargetPlatform) -cookflavor=&quot;$(BCRCookFlavor)&quot;" />
 			<Property Name="BCRArgs" Value="$(BCRArgs) -client" If="'$(TargetType)' == 'Client'" />
 			<Property Name="BCRArgs" Value="$(BCRArgs) -server" If="'$(DedicatedServer)'" />
 			
