EngineTest - Use new Horde device pool service

#rnx

diff --git a/Engine/Build/Graph/Tasks/Inc/GauntletSettings.xml b/Engine/Build/Graph/Tasks/Inc/GauntletSettings.xml
--- a/Engine/Build/Graph/Tasks/Inc/GauntletSettings.xml
+++ b/Engine/Build/Graph/Tasks/Inc/GauntletSettings.xml
@@ -19,11 +19,12 @@
 
 	<!-- How long to give Gauntlet before a test is considered timed out (some tests will time out earlier if they detect inactivity)-->
 	<Option Name="GauntletTimeout" DefaultValue="3600" Description="Max running time of a test in Gauntlet"/>
+	<Option Name="GauntletDevicePool" DefaultValue="UE4" Description="Gauntlet Device Pool to use to run target platform tests"/>
 	
 	<Property Name="GauntletCommonArgs" Value="-maxduration=$(GauntletTimeOut) -unattended -verbose" />
 
 	<Property Name="GauntletDeviceArgs" Value="" />
-	<Property Name="GauntletDeviceArgs" Value="-deviceurl=http://deviceservice.epicgames.net -reboot -cleardevices -sandbox=&quot;$(ProjectName)&quot;" If ="$(IsBuildMachine)" />
+	<Property Name="GauntletDeviceArgs" Value="-deviceurl=https://horde.devtools.epicgames.com/ -devicepool=&quot;$(GauntletDevicePool)&quot; -reboot -cleardevices -sandbox=&quot;$(ProjectName)&quot;" If="$(IsBuildMachine)" />
 
 	<Property Name="GauntletTargetTestArgs" Value="$(GauntletCommonArgs) $(GauntletDeviceArgs)" />
 	<Property Name="GauntletEditorTestArgs" Value="-build=editor $(GauntletCommonArgs)" />
