Exclude previous steps from labels for package, publish, and boot test so they do not inherit warning states

diff --git a/Engine/Build/Graph/Tasks/BuildAndTestProject.xml b/Engine/Build/Graph/Tasks/BuildAndTestProject.xml
--- a/Engine/Build/Graph/Tasks/BuildAndTestProject.xml
+++ b/Engine/Build/Graph/Tasks/BuildAndTestProject.xml
@@ -397,9 +397,9 @@
 		<!-- Declare labels for CIS -->
 		<Label Category="Platforms" Name="$(ProjectName) $(TargetPlatform)" Requires="$(PlatformCompileNodeName)" />
 		<Label Category="Cook" Name="Cook $(DisplayTargetName) $(TargetPlatform)" Requires="$(PlatformCookNodeName)" />
-		<Label Category="Stage" Name="Stage $(DisplayTargetName) $(TargetPlatform)" Requires="$(PlatformStageNodeName);$(PlatformHashNodeName)" />
-		<Label Category="Package" Name="Pkg $(DisplayTargetName) $(TargetPlatform)" Requires="$(PlatformPackageNodeName);$(PlatformPublishPackagedNodeName)"/>
-		<Label Category="Test" Name="BootTest $(DisplayTargetName) $(TargetPlatform)" Requires="$(AllPlatformBootTestNodes)" Exclude="$(PlatformEditorNodeName)" />
+		<Label Category="Stage" Name="Stage $(DisplayTargetName) $(TargetPlatform)" Requires="$(PlatformStageNodeName);$(PlatformHashNodeName);$(PlatformPublishStagedNodeName)" Exclude="$(PlatformCookNodeName)" />
+		<Label Category="Package" Name="Pkg $(DisplayTargetName) $(TargetPlatform)" Requires="$(PlatformPackageNodeName);$(PlatformPublishPackagedNodeName)" Exclude="$(PlatformPublishStagedNodeName)" />
+		<Label Category="Test" Name="BootTest $(DisplayTargetName) $(TargetPlatform)" Requires="$(AllPlatformBootTestNodes)" Exclude="$(PlatformEditorNodeName);$(PlatformPackageNodeName);$(PlatformPublishPackagedNodeName)" />
 		<Label Category="Test" Name="Test $(DisplayTargetName) $(TargetPlatform)" Requires="$(AllPlatformTestNodes)" Exclude="$(PlatformEditorNodeName);$(AllPlatformBootTestNodes)"/>
 
 		<!-- Tie all steps into our platform tasks unless excluded -->
