UE.Automation - Enable resume on critical failure for project that uses BuildAndTestProject buildgraph

#rnx

diff --git a/Engine/Build/Graph/Tasks/Inc/GauntletSettings.xml b/Engine/Build/Graph/Tasks/Inc/GauntletSettings.xml
--- a/Engine/Build/Graph/Tasks/Inc/GauntletSettings.xml
+++ b/Engine/Build/Graph/Tasks/Inc/GauntletSettings.xml
@@ -21,7 +21,7 @@
 	<Option Name="GauntletTimeout" DefaultValue="3600" Description="Max running time of a test in Gauntlet"/>
 	<Option Name="GauntletDevicePool" DefaultValue="UE4" Description="Gauntlet Device Pool to use to run target platform tests"/>
 	
-	<Property Name="GauntletCommonArgs" Value="-maxduration=$(GauntletTimeOut) -unattended -verbose" />
+	<Property Name="GauntletCommonArgs" Value="-maxduration=$(GauntletTimeOut) -unattended -verbose -ResumeOnCriticalFailure" />
 
 	<Property Name="GauntletDeviceArgs" Value="" />
 	<Property Name="GauntletDeviceArgs" Value="-deviceurl=https://horde.devtools.epicgames.com/ -devicepool=&quot;$(GauntletDevicePool)&quot; -reboot -cleardevices -sandbox=&quot;$(ProjectName)&quot;" If="$(IsBuildMachine)" />
