Enforce alternate device pool for some platforms.

#rnx

diff --git a/Engine/Build/Graph/Tasks/BuildAndTestProject.xml b/Engine/Build/Graph/Tasks/BuildAndTestProject.xml
--- a/Engine/Build/Graph/Tasks/BuildAndTestProject.xml
+++ b/Engine/Build/Graph/Tasks/BuildAndTestProject.xml
@@ -353,7 +353,10 @@
 					</Do>
 					<Property Name="GauntletBuildArg" Value="-build=&quot;$(TestDirectory)&quot;" />
 				</Do>
-				
+
+				<Property Name="GauntletTestArgs" Value="$(GauntletTargetTestArgs)"/>
+				<Property Name="GauntletTestArgs" Value="$(GauntletTargetTestAltArgs)" If="ContainsItem('$(AlternateDevicePoolPlatforms)', '$(TargetPlatform)', ';')"/>
+
 				<!-- Declare a boot test for the target-->	
 				<Property Name="BootTestName" Value="$(DisplayTargetName) $(TargetPlatform) $(TestConfiguration) BootTest" />
 				<Node Name="$(BootTestName)" After="$(PlatformPublishPackagedNodeName)" Requires="$(TestPrerequisites)">
@@ -364,7 +367,7 @@
 						<Property Name="GauntletReportArgs" Value="$(GauntletReportArgs) -ReportURL=&quot;$(AutomationReportUrl)/$(OutputFragmentPath)/&quot;" If="'$(AutomationReportUrl)' != ''" />
 					</Do>
 					<!-- Run the boot test -->
-					<Command Name="RunUnreal" Arguments="-project=$(TargetProject) -test=&quot;UE.BootTest&quot; -platform=$(TargetPlatform) -configuration=$(TestConfiguration) $(GauntletBuildArg) $(GauntletTargetTestArgs) $(GauntletReportArgs)" />
+					<Command Name="RunUnreal" Arguments="-project=$(TargetProject) -test=&quot;UE.BootTest&quot; -platform=$(TargetPlatform) -configuration=$(TestConfiguration) $(GauntletBuildArg) $(GauntletTestArgs) $(GauntletReportArgs)" />
 				</Node>
 				<!-- Link the boot test in to the prerequisites and generated nodes if its enabled-->
 				<Property Name="TestPrerequisites" Value="$(TestPrerequisites);$(PlatformEditorNodeName)" />
@@ -384,7 +387,7 @@
 					<Property Name="TestNodeName" Value="$(DisplayTargetName) $(TargetPlatform) $(TestConfiguration) $(TestName)" />
 					<!-- Run the test. No 'Requires' as we run on the same machine locally and use the publish folder otherwise -->
 					<Node Name="$(TestNodeName)" Requires="$(TestPrerequisites)">
-						<Command Name="RunUnreal" Arguments="-project=$(TargetProject) -test=&quot;$(TestName)&quot; -platform=$(TargetPlatform) -configuration=$(TestConfiguration) $(GauntletBuildArg) $(GauntletTargetTestArgs) $(GauntletReportArgs)"/>
+						<Command Name="RunUnreal" Arguments="-project=$(TargetProject) -test=&quot;$(TestName)&quot; -platform=$(TargetPlatform) -configuration=$(TestConfiguration) $(GauntletBuildArg) $(GauntletTestArgs) $(GauntletReportArgs)"/>
 					</Node>
 					<!-- Add it to our list -->
 					<Do If="!ContainsItem('$(TargetPlatformsToSkipTests)', '$(TargetPlatform)', '+')">
diff --git a/Engine/Build/Graph/Tasks/Inc/GauntletSettings.xml b/Engine/Build/Graph/Tasks/Inc/GauntletSettings.xml
--- a/Engine/Build/Graph/Tasks/Inc/GauntletSettings.xml
+++ b/Engine/Build/Graph/Tasks/Inc/GauntletSettings.xml
@@ -8,9 +8,7 @@
 		This script should be includedEng after CommonProperties.xml and PathSettings.xml where build
         strings and paths are defines
 		
-		InProperty: ProjectPath - Path where the .uproject is located
         InProperty: ProjectName - Name of the project (e.g. <ProjectName>.uproject)
-		InProperty: EscapedBranch
 
         OutProperty: GauntletEditorTestArgs - Arguments to pass to Gauntlet tests executed under the editor
 		OutProperty: GauntletTargetTestArgs - Arguments to pass to Gauntlet tests executed using cooked builds
@@ -20,13 +18,19 @@
 	<!-- How long to give Gauntlet before a test is considered timed out (some tests will time out earlier if they detect inactivity)-->
 	<Option Name="GauntletTimeout" DefaultValue="3600" Description="Max running time of a test in Gauntlet"/>
 	<Option Name="GauntletDevicePool" DefaultValue="UE4" Description="Gauntlet Device Pool to use to run target platform tests"/>
+	<Option Name="AlternateGauntletDevicePool" DefaultValue="UE5" Description="Alternate Gauntlet Device Pool to use to run target platform tests"/>
+	<Option Name="AlternateDevicePoolPlatforms" DefaultValue="XSX;PS5" Description="Platforms that need to use the Alternate Gauntlet Device Pool"/>
 	
 	<Property Name="GauntletCommonArgs" Value="-maxduration=$(GauntletTimeOut) -unattended -verbose -ResumeOnCriticalFailure" />
 
 	<Property Name="GauntletDeviceArgs" Value="" />
-	<Property Name="GauntletDeviceArgs" Value="-deviceurl=https://horde.devtools.epicgames.com/ -devicepool=&quot;$(GauntletDevicePool)&quot; -reboot -cleardevices -sandbox=&quot;$(ProjectName)&quot;" If="$(IsBuildMachine)" />
+	<Property Name="GauntletDeviceArgs" Value="-deviceurl=https://horde.devtools.epicgames.com/ -reboot -cleardevices -sandbox=&quot;$(ProjectName)&quot;" If="$(IsBuildMachine)" />
+	<Property Name="GauntletDeviceAltArgs" Value="$(GauntletDeviceArgs)" />
+	<Property Name="GauntletDeviceArgs" Value="$(GauntletDeviceArgs) -devicepool=&quot;$(GauntletDevicePool)&quot;" If="$(IsBuildMachine)"/>
+	<Property Name="GauntletDeviceAltArgs" Value="$(GauntletDeviceAltArgs) -devicepool=&quot;$(AlternateGauntletDevicePool)&quot;" If="$(IsBuildMachine)"/>
 
 	<Property Name="GauntletTargetTestArgs" Value="$(GauntletCommonArgs) $(GauntletDeviceArgs)" />
+	<Property Name="GauntletTargetTestAltArgs" Value="$(GauntletCommonArgs) $(GauntletDeviceAltArgs)" />
 	<Property Name="GauntletEditorTestArgs" Value="-build=editor $(GauntletCommonArgs)" />
 
     <!-- this is fully generated dynamiclly when needed, it's simply here so it can be used without worrying if its set -->
