Fill editor DDC before running editor test to improve test execution performance.

#rnx

diff --git a/Engine/Build/Graph/Tasks/BuildAndTestProject.xml b/Engine/Build/Graph/Tasks/BuildAndTestProject.xml
--- a/Engine/Build/Graph/Tasks/BuildAndTestProject.xml
+++ b/Engine/Build/Graph/Tasks/BuildAndTestProject.xml
@@ -71,6 +71,8 @@
 	<Option Name="EditorBootTest" DefaultValue="!$(SkipTest)" Description="Run a simple boot test on the editor before other tests"/>
 	<Option Name="TargetBootTest" DefaultValue="!$(SkipTest)" Description="Run a simple boot test on target platforms before other tests"/>
 
+	<Option Name="FillEditorDDC" DefaultValue="false" Description="Whether to fill the editor DDC before running any editor tests"/>
+
 	<Option Name="EditorTestList" DefaultValue="$(DefaultEditorTestList)" Description="Tests to run on the editor. Test1+Test2 etc"/>
 	<Option Name="TargetTestList" DefaultValue="$(DefaultTargetTestList)" Description="Tests to run on the target platform(s). Test1+Test2 etc"/>
 	<Option Name="EditorPlatformsToSkipBootTests" DefaultValue="$(EditorPlatformsToSkipTests)" Description="Editor platforms that should skip boot tests."/>
@@ -193,6 +195,12 @@
 			<Node Name="$(PlatformBootTestNodeName)" Requires="$(GeneratedToolNodes)">
 				<Command Name="RunUnreal" Arguments="-project=$(TargetProject) -build=Editor -test=&quot;UE.EditorBootTest&quot; $(GauntletEditorTestArgs) $(GauntletReportArgs)" />
 			</Node>
+			<Property Name="ProjectEditorDDC" Value="$(ProjectName)Editor DDC $(EditorPlatform)" />
+			<Node Name="$(ProjectEditorDDC)" Requires="$(GeneratedToolNodes)">
+				<Property Name="DDCTargetPlatform" Value="$(EditorPlatform)" />
+				<Property Name="DDCTargetPlatform" Value="Windows" If="'$(DDCTargetPlatform)' == 'Win64'"/>
+				<Commandlet Name="DerivedDataCache" Project="$(RootDir)/$(TargetProject)" Arguments="-fill -targetplatform=$(DDCTargetPlatform)Editor"/>
+			</Node>
 			<!-- Link the boot test in to the prerequisites and generated nodes if its enabled-->
 			<Do If="$(EditorBootTest)">
 				<Property Name="TestPrerequisites" Value="$(TestPrerequisites);$(PlatformBootTestNodeName)" />
@@ -215,12 +223,14 @@
 
 			<!-- Create Labels -->
 			<Label Category="Editors" Name="$(ProjectName) Editor $(EditorPlatform)" Requires="$(GeneratedToolNodes)" />
+			<Label Category="Editors" Name="$(ProjectName) Editor DDC $(EditorPlatform)" Requires="$(ProjectEditorDDC)" />
 
 			<!-- If this editor was requested (as opposed to being required due to target platforms) then link in the test nodes and set up labels -->
 			<Do If="ContainsItem('$(EditorPlatforms)', '$(EditorPlatform)', '+') and !ContainsItem('$(EditorPlatformsToSkipTests)', '$(EditorPlatform)', '+')">
 				<Label Category="Test" Name="BootTest $(ProjectName)Editor $(EditorPlatform)" Requires="$(PlatformBootTestNodeName)" />
 				<Label Category="Test" Name="Test $(ProjectName)Editor $(EditorPlatform)" Requires="$(PlatformEditorTestNodes)" />
 				<Property Name="GeneratedEditorTestNodes" Value="$(GeneratedEditorTestNodes);$(PlatformBootTestNodeName)" If="$(EditorBootTest)" />
+				<Property Name="GeneratedEditorTestNodes" Value="$(GeneratedEditorTestNodes);$(ProjectEditorDDC)" If="$(FillEditorDDC)" />
 				<Property Name="GeneratedEditorTestNodes" Value="$(GeneratedEditorTestNodes);$(PlatformEditorTestNodes)" />
 			</Do>
 		</Agent>
