Fix build script not passing the build name to UBT, causing preflights not to be labeled as such.

diff --git a/Engine/Build/Graph/Tasks/BuildAndTestProject.xml b/Engine/Build/Graph/Tasks/BuildAndTestProject.xml
--- a/Engine/Build/Graph/Tasks/BuildAndTestProject.xml
+++ b/Engine/Build/Graph/Tasks/BuildAndTestProject.xml
@@ -304,9 +304,9 @@
 				<ForEach Name="TargetConfiguration" Values="$(TargetConfigurations)" Separator="+">
 					
 					<!-- Build the binaries -->
-					<Compile Target="$(TargetName)" Platform="$(TargetPlatform)" Configuration="$(TargetConfiguration)" Arguments="$(GenericCompileArguments) $(ExtraTargetCompileArguments)" Tag="#BinaryFiles_$(TargetPlatform)_$(TargetConfiguration)"/>
+					<Compile Target="$(TargetName)" Platform="$(TargetPlatform)" Configuration="$(TargetConfiguration)" Arguments="-BuildVersion=&quot;$(BuildName)&quot; $(GenericCompileArguments) $(ExtraTargetCompileArguments)" Tag="#BinaryFiles_$(TargetPlatform)_$(TargetConfiguration)"/>
 					<Do If="'$(DedicatedServer)'" >
-						<Compile Target="$(TargetName)Server" Platform="Win64" Configuration="$(TargetConfiguration)" Arguments="$(GenericCompileArguments) $(ExtraTargetCompileArguments)" Tag="#BinaryFiles_$(TargetPlatform)_$(TargetConfiguration)"/>
+						<Compile Target="$(TargetName)Server" Platform="Win64" Configuration="$(TargetConfiguration)" Arguments="-BuildVersion=&quot;$(BuildName)&quot; $(GenericCompileArguments) $(ExtraTargetCompileArguments)" Tag="#BinaryFiles_$(TargetPlatform)_$(TargetConfiguration)"/>
 					</Do>
 					
 					<!-- Upload binaries to symbol store, if enabled -->
