Fix build parallelism
 - Symbol upload tasks mixed with Compile tasks were causing compiles to happen sequentially rather than in parallel.
 - Also, separated compile and cook/stage/package nodes into two agents so they run as soon as their dependencies are complete.

diff --git a/Engine/Build/Graph/Tasks/BuildAndTestProject.xml b/Engine/Build/Graph/Tasks/BuildAndTestProject.xml
--- a/Engine/Build/Graph/Tasks/BuildAndTestProject.xml
+++ b/Engine/Build/Graph/Tasks/BuildAndTestProject.xml
@@ -269,6 +269,8 @@
 		<!-- Define these names up front as we use some of them several times -->
 		<Property Name="PlatformAgentName" Value="$(DisplayTargetName) $(TargetPlatform)" />
 		<Property Name="PlatformCompileNodeName" Value="Compile $(TargetName) $(TargetPlatform)" />
+		<Property Name="UploadSymbolsNodeName" Value="" />
+		<Property Name="UploadSymbolsNodeName" Value="Upload Symbols $(TargetName) $(TargetPlatform)" If="'$(SymbolServerPath)' != ''" />
 		<Property Name="PlatformCookNodeName" Value="Cook $(DisplayTargetName) $(TargetPlatform)" />
 		<Property Name="PlatformStageNodeName" Value="Stage $(DisplayTargetName) $(TargetPlatform)" />
 		<Property Name="PlatformPackageNodeName" Value="Package $(DisplayTargetName) $(TargetPlatform)" />
@@ -290,33 +292,38 @@
 
 		<!-- Hashing of staged data can only be done on Win64 hosts -->
 		<Property Name="DoUnsyncHashing" Value="$(HostAgentType) == 'Win64'" />
+
+		<!-- Default BuildCookRun args needed when staging/packaging/archiving this configuration -->
+		<Property Name="BCRArgs" Value="-nocompile -project=$(TargetProject) -platform=$(TargetPlatform) -cookflavor=&quot;$(BCRCookFlavor)&quot;" />
+		<Property Name="BCRArgs" Value="$(BCRArgs) -client" If="'$(TargetType)' == 'Client'" />
+		<Property Name="BCRArgs" Value="$(BCRArgs) -server" If="'$(DedicatedServer)'" />
 		
 		<!-- Build the client and if necessary a server-->
-		<Agent Name="Make $(PlatformAgentName)" Type="$(BuildAgentType)">
+		<Agent Name="Compile $(PlatformAgentName)" Type="$(BuildAgentType)">
 
-			<!-- Default BuildCookRun args needed when staging/packaging/archiving this configuration -->
-			<Property Name="BCRArgs" Value="-nocompile -project=$(TargetProject) -platform=$(TargetPlatform) -cookflavor=&quot;$(BCRCookFlavor)&quot;" />
-			<Property Name="BCRArgs" Value="$(BCRArgs) -client" If="'$(TargetType)' == 'Client'" />
-			<Property Name="BCRArgs" Value="$(BCRArgs) -server" If="'$(DedicatedServer)'" />
-			
 			<!-- Compile -->
-			<Node Name="$(PlatformCompileNodeName)" Requires="$(PlatformToolsNodeName)" RunEarly="true">
+			<Node Name="$(PlatformCompileNodeName)" Requires="$(PlatformToolsNodeName)">
 				<ForEach Name="TargetConfiguration" Values="$(TargetConfigurations)" Separator="+">
-					
-					<!-- Build the binaries -->
-					<Compile Target="$(TargetName)" Platform="$(TargetPlatform)" Configuration="$(TargetConfiguration)" Arguments="-BuildVersion=&quot;$(BuildName)&quot; $(GenericCompileArguments) $(ExtraTargetCompileArguments)" Tag="#BinaryFiles_$(TargetPlatform)_$(TargetConfiguration)"/>
+					<Compile Target="$(TargetName)" Platform="$(TargetPlatform)" Configuration="$(TargetConfiguration)" Arguments="-BuildVersion=&quot;$(BuildName)&quot; $(GenericCompileArguments) $(ExtraTargetCompileArguments)"/>
 					<Do If="'$(DedicatedServer)'" >
-						<Compile Target="$(TargetName)Server" Platform="Win64" Configuration="$(TargetConfiguration)" Arguments="-BuildVersion=&quot;$(BuildName)&quot; $(GenericCompileArguments) $(ExtraTargetCompileArguments)" Tag="#BinaryFiles_$(TargetPlatform)_$(TargetConfiguration)"/>
-					</Do>
-					
-					<!-- Upload binaries to symbol store, if enabled -->
-					<Do If="'$(SymbolServerPath)' != ''">
-						<Log Message="Uploading symbols to &quot;$(SymbolServerPath)&quot; ..."/>
-						<SymStore Platform="$(TargetPlatform)" Files="#BinaryFiles_$(TargetPlatform)_$(TargetConfiguration)" StoreDir="$(SymbolServerPath)/$(TargetPlatform)" Product="$(TargetName)" BuildVersion="$(BuildName)" />
+						<Compile Target="$(TargetName)Server" Platform="Win64" Configuration="$(TargetConfiguration)" Arguments="-BuildVersion=&quot;$(BuildName)&quot; $(GenericCompileArguments) $(ExtraTargetCompileArguments)"/>
 					</Do>
 				</ForEach>
 			</Node>
 
+			<!-- (Optional) Upload Symbols -->
+			<Do If="'$(UploadSymbolsNodeName)' != ''">
+				<Node Name="$(UploadSymbolsNodeName)" Requires="$(PlatformCompileNodeName)">
+					<Log Message="Uploading symbols to &quot;$(SymbolServerPath)&quot; ..."/>
+					<SymStore Platform="$(TargetPlatform)" Files="#$(PlatformCompileNodeName)" StoreDir="$(SymbolServerPath)/$(TargetPlatform)" Product="$(TargetName)" BuildVersion="$(BuildName)" />
+				</Node>
+			</Do>
+
+		</Agent>
+
+		<!-- Separate agent for cooking / staging, so the cook can happen in parallel with the client compile -->
+		<Agent Name="Package $(PlatformAgentName)" Type="$(BuildAgentType)">
+
 			<!-- Cook -->
 			<Node Name="$(PlatformCookNodeName)" Requires="$(PlatformEditorNodeName)" Produces="#$(PlatformCookNodeName) Complete">
 				<Property Name="CookArgs" Value="" />
@@ -418,7 +425,7 @@
 		</Agent>
 
 		<!-- Declare labels for CIS -->
-		<Label Category="Platforms" Name="$(ProjectName) $(TargetPlatform)" Requires="$(PlatformCompileNodeName)" />
+		<Label Category="Platforms" Name="$(ProjectName) $(TargetPlatform)" Requires="$(PlatformCompileNodeName);$(UploadSymbolsNodeName)" />
 		<Label Category="Cook" Name="Cook $(DisplayTargetName) $(TargetPlatform)" Requires="$(PlatformCookNodeName)" />
 		<Label Category="Stage" Name="Stage $(DisplayTargetName) $(TargetPlatform)" Requires="$(PlatformStageNodeName);$(PlatformHashNodeName);$(PlatformPublishStagedNodeName)" Exclude="$(PlatformCookNodeName)" />
 		<Label Category="Package" Name="Pkg $(DisplayTargetName) $(TargetPlatform)" Requires="$(PlatformPackageNodeName);$(PlatformPublishPackagedNodeName)" Exclude="$(PlatformPublishStagedNodeName)" />
@@ -427,7 +434,7 @@
 
 		<!-- Tie all steps into our platform tasks unless excluded -->
 		<Property Name="PlatformTasks" Value="" />
-		<Property Name="PlatformTasks" Value="$(PlatformCompileNodeName)" If="!$(SkipBuild)" />
+		<Property Name="PlatformTasks" Value="$(PlatformCompileNodeName);$(UploadSymbolsNodeName)" If="!$(SkipBuild)" />
 		<Property Name="PlatformTasks" Value="$(PlatformTasks);$(PlatformCookNodeName)" If="!$(SkipCook) and !$(SkipPublish)" />		
 		<Property Name="PlatformTasks" Value="$(PlatformTasks);$(PlatformPackageNodeName);$(PlatformPublishPackagedNodeName)" If="!$(SkipPublish) And $(PublishBuilds)" />	
 		<Property Name="PlatformTasks" Value="$(PlatformTasks);$(PlatformHashNodeName)"  If="$(PublishHashedStagingData)" />	
