Add agent override for client compile steps.
 - This allows us to put C++ compilation on dedicated agents with faster single thread performance, e.g. for long PGO/LTO links.

Add RunEarly to cook node.
 - Allows the cook to start before client binaries have completed.

diff --git a/Engine/Build/Graph/Tasks/BuildAndTestProject.xml b/Engine/Build/Graph/Tasks/BuildAndTestProject.xml
--- a/Engine/Build/Graph/Tasks/BuildAndTestProject.xml
+++ b/Engine/Build/Graph/Tasks/BuildAndTestProject.xml
@@ -21,6 +21,7 @@
 	<Option Name="UseIncrementalAgents" DefaultValue="false" Restrict="true|false" Description="Use incremental agents for building" />
 	<Option Name="AgentOverride" DefaultValue="" Description="If set, other logic is ignored and this agent type is used for all non-test work." />
 	<Option Name="TestAgentOverride" DefaultValue="" Description="If set, other logic is ignored and this agent type is used for all test work." />
+	<Option Name="ClientCompileAgentOverride" DefaultValue="" Description="Specifies the agent type to use for C++ compilation jobs for target platform binaries." />
 
 	<Option Name="WithBATDefaults" DefaultValue="true" Description="Used by parent scripts to specify defaults when including this script"/>
 
@@ -256,11 +257,14 @@
 		<Property Name="TestAgentType" Value="Test$(HostAgentType)" If="'$(TargetPlatform)' == 'Win64' or '$(TargetPlatform)' == 'Mac'"/>
 
 		<!-- Use incremental agents for building? -->
-		<Property Name="BuildAgentType" Value="$(HostAgentType)" />
-		<Property Name="BuildAgentType" Value="Incremental$(HostAgentType)" If="$(UseIncrementalAgents)" />
+		<Property Name="CompileAgentType" Value="$(HostAgentType)" />
+		<Property Name="CompileAgentType" Value="Incremental$(HostAgentType)" If="$(UseIncrementalAgents)" />
 
 		<!-- Use agent types from override? -->
-		<Property Name="BuildAgentType" Value="$(AgentOverride)" If="'$(AgentOverride)' != ''"/>
+		<Property Name="CompileAgentType" Value="$(AgentOverride)" If="'$(AgentOverride)' != ''"/>
+		<!-- Use the same agent type for packaging as compilation -->
+		<Property Name="PackageAgentType" Value="$(CompileAgentType)"/>
+		<Property Name="CompileAgentType" Value="$(ClientCompileAgentOverride)" If="'$(ClientCompileAgentOverride)' != ''"/>
 		<Property Name="TestAgentType" Value="$(TestAgentOverride)" If="'$(TestAgentOverride)' != ''"/>
 
 		<!-- Building and cooking will require these nodes from the appropriate host. They should have been built above... -->
@@ -299,7 +303,7 @@
 		<Property Name="BCRArgs" Value="$(BCRArgs) -server" If="'$(DedicatedServer)'" />
 		
 		<!-- Build the client and if necessary a server-->
-		<Agent Name="Compile $(PlatformAgentName)" Type="$(BuildAgentType)">
+		<Agent Name="Compile $(PlatformAgentName)" Type="$(CompileAgentType)">
 
 			<!-- Compile -->
 			<Node Name="$(PlatformCompileNodeName)" Requires="$(PlatformToolsNodeName)">
@@ -322,10 +326,10 @@
 		</Agent>
 
 		<!-- Separate agent for cooking / staging, so the cook can happen in parallel with the client compile -->
-		<Agent Name="Package $(PlatformAgentName)" Type="$(BuildAgentType)">
+		<Agent Name="Package $(PlatformAgentName)" Type="$(PackageAgentType)">
 
 			<!-- Cook -->
-			<Node Name="$(PlatformCookNodeName)" Requires="$(PlatformEditorNodeName)" Produces="#$(PlatformCookNodeName) Complete">
+			<Node Name="$(PlatformCookNodeName)" Requires="$(PlatformEditorNodeName)" Produces="#$(PlatformCookNodeName) Complete" RunEarly="true">
 				<Property Name="CookArgs" Value="" />
 				<Property Name="CookArgs" Value="-MapIniSection=$(MapIniSection) $(CookPlatform)" If="'$(MapIniSection)' != ''"/>				
 				<Cook Project="$(TargetProject)" Platform="$(CookPlatform)" Arguments="$(CookArgs) $(ExtraCookArguments)"/>
