Add option for producing packaged client builds using "-distribution" UAT command line.
 - Required for creating packages that can be shipped to retail.

diff --git a/Engine/Build/Graph/Tasks/BuildAndTestProject.xml b/Engine/Build/Graph/Tasks/BuildAndTestProject.xml
--- a/Engine/Build/Graph/Tasks/BuildAndTestProject.xml
+++ b/Engine/Build/Graph/Tasks/BuildAndTestProject.xml
@@ -23,6 +23,8 @@
 	<Option Name="TestAgentOverride" DefaultValue="" Description="If set, other logic is ignored and this agent type is used for all test work." />
 	<Option Name="ClientCompileAgentOverride" DefaultValue="" Description="Specifies the agent type to use for C++ compilation jobs for target platform binaries." />
 
+	<Option Name="CreateDistributionPackages" DefaultValue="false" Description="Enables the use of &quot;-distribution&quot; in the UAT client staging/packaging steps. Required when producing builds that are intended for retail distribution."/>
+
 	<Option Name="WithBATDefaults" DefaultValue="true" Description="Used by parent scripts to specify defaults when including this script"/>
 
 	<!-- 
@@ -298,9 +300,10 @@
 		<Property Name="DoUnsyncHashing" Value="$(HostAgentType) == 'Win64'" />
 
 		<!-- Default BuildCookRun args needed when staging/packaging/archiving this configuration -->
-		<Property Name="BCRArgs" Value="-nocompile -project=$(TargetProject) -platform=$(TargetPlatform) -cookflavor=&quot;$(BCRCookFlavor)&quot;" />
+		<Property Name="BCRArgs" Value="--nocompile -project=$(TargetProject) -platform=$(TargetPlatform) -cookflavor=&quot;$(BCRCookFlavor)&quot; -configuration=$(PackageConfigurations)" />
 		<Property Name="BCRArgs" Value="$(BCRArgs) -client" If="'$(TargetType)' == 'Client'" />
 		<Property Name="BCRArgs" Value="$(BCRArgs) -server" If="'$(DedicatedServer)'" />
+		<Property Name="BCRArgs" Value="$(BCRArgs) -distribution" If="$(CreateDistributionPackages)"/>
 		
 		<!-- Build the client and if necessary a server-->
 		<Agent Name="Compile $(PlatformAgentName)" Type="$(CompileAgentType)">
@@ -337,7 +340,7 @@
 			
 			<!-- Stage -->
 			<Node Name="$(PlatformStageNodeName)" Requires="$(PlatformCompileNodeName);#$(PlatformCookNodeName) Complete">
-				<Command Name="BuildCookRun" Arguments="$(BCRArgs) $(ExtraStageAndPackageArguments) -configuration=$(PackageConfigurations) -skipbuild -skipcook -stage -pak" />
+				<Command Name="BuildCookRun" Arguments="$(BCRArgs) $(ExtraStageAndPackageArguments) -skipbuild -skipcook -stage -pak" />
 			</Node>
 			
 			<!-- Hash -->
@@ -355,12 +358,12 @@
 			
 			<!-- Package -->
 			<Node Name="$(PlatformPackageNodeName)" Requires="$(PlatformPublishStagedNodeName)">
-				<Command Name="BuildCookRun" Arguments="$(BCRArgs) $(ExtraStageAndPackageArguments) -configuration=$(PackageConfigurations) -skipbuild -skipcook -skipstage -skippak -package" />
+				<Command Name="BuildCookRun" Arguments="$(BCRArgs) $(ExtraStageAndPackageArguments) -skipbuild -skipcook -skipstage -skippak -package" />
 			</Node>
 			
 			<!-- Publish (Packages) -->
 			<Node Name="$(PlatformPublishPackagedNodeName)" Requires="$(PlatformPackageNodeName)">
-				<Command Name="BuildCookRun" Arguments="$(BCRArgs) $(ExtraStageAndPackageArguments) -configuration=$(PackageConfigurations) -skipbuild -skipcook -skipstage -skippak -skippackage -archive -archivedirectory=&quot;$(NetworkOutputDirectory)/$(StagedPlatformFolder)/Packages&quot;" />
+				<Command Name="BuildCookRun" Arguments="$(BCRArgs) $(ExtraStageAndPackageArguments) -skipbuild -skipcook -skipstage -skippak -skippackage -archive -archivedirectory=&quot;$(NetworkOutputDirectory)/$(StagedPlatformFolder)/Packages&quot;" />
 			</Node>
 			
 		</Agent>
