Fix for BuildAndTestProject.xml when not using FrostyTests override script

diff --git a/Engine/Build/Graph/Tasks/BuildAndTestProject.xml b/Engine/Build/Graph/Tasks/BuildAndTestProject.xml
--- a/Engine/Build/Graph/Tasks/BuildAndTestProject.xml
+++ b/Engine/Build/Graph/Tasks/BuildAndTestProject.xml
@@ -16,6 +16,8 @@
 	<Option Name="MapIniSection" DefaultValue="" Restrict="" Description="MapIniSection argument to pass when cooking" />
 	<Option Name="ExtraCookArguments" DefaultValue="" Restrict="" Description="Extra arguments to use when cooking" />
 	<Option Name="ExtraStageAndPackageArguments" DefaultValue="" Restrict="" Description="Extra arguments to use when staging and packaging" />
+	<Option Name="TargetPlatformFilteredExtraStageAndPackageArguments" DefaultValue="" Restrict="" Description="Extra arguments to use when staging and packaging.  Filtered to specific platforms via TargetPlatformFilteredExtraStageAndPackageArgumentsFilter." />
+	<Option Name="TargetPlatformFilteredExtraStageAndPackageArgumentsFilter" DefaultValue="" Restrict="" Description="Platform filter applied to TargetPlatformFilteredExtraStageAndPackageArguments." />
 	<Option Name="WithRemoteExecWorkers" DefaultValue="false" Restrict="" Description="If your project uses standalone remote execution workers" />
 
 	<Option Name="UseIncrementalAgents" DefaultValue="false" Restrict="true|false" Description="Use incremental agents for building" />
