Exclude package warnings from test labels

diff --git a/Engine/Build/Graph/Tasks/BuildAndTestProject.xml b/Engine/Build/Graph/Tasks/BuildAndTestProject.xml
--- a/Engine/Build/Graph/Tasks/BuildAndTestProject.xml
+++ b/Engine/Build/Graph/Tasks/BuildAndTestProject.xml
@@ -440,7 +440,7 @@
 		<Label Category="Stage" Name="Stage $(DisplayTargetName) $(TargetPlatform)" Requires="$(PlatformStageNodeName);$(PlatformHashNodeName);$(PlatformPublishStagedNodeName)" Exclude="$(PlatformCookNodeName)" />
 		<Label Category="Package" Name="Pkg $(DisplayTargetName) $(TargetPlatform)" Requires="$(PlatformPackageNodeName);$(PlatformPublishPackagedNodeName)" Exclude="$(PlatformPublishStagedNodeName)" />
 		<Label Category="Test" Name="BootTest $(DisplayTargetName) $(TargetPlatform)" Requires="$(AllPlatformBootTestNodes)" Exclude="$(PlatformEditorNodeName);$(PlatformPackageNodeName);$(PlatformPublishPackagedNodeName)" />
-		<Label Category="Test" Name="Test $(DisplayTargetName) $(TargetPlatform)" Requires="$(AllPlatformTestNodes)" Exclude="$(PlatformEditorNodeName);$(AllPlatformBootTestNodes)"/>
+		<Label Category="Test" Name="Test $(DisplayTargetName) $(TargetPlatform)" Requires="$(AllPlatformTestNodes)" Exclude="$(PlatformEditorNodeName);$(PlatformPackageNodeName);$(AllPlatformBootTestNodes)"/>
 
 		<!-- Tie all steps into our platform tasks unless excluded -->
 		<Property Name="PlatformTasks" Value="" />
