Include the original GauntletSettings.xml from the override settings
Add extra options to be able to add Gauntlet arguments

#rnx

diff --git a/Engine/Build/Graph/Tasks/Inc/GauntletSettings.xml b/Engine/Build/Graph/Tasks/Inc/GauntletSettings.xml
--- a/Engine/Build/Graph/Tasks/Inc/GauntletSettings.xml
+++ b/Engine/Build/Graph/Tasks/Inc/GauntletSettings.xml
@@ -20,8 +20,11 @@
 	<Option Name="GauntletDevicePool" DefaultValue="UE4" Description="Gauntlet Device Pool to use to run target platform tests"/>
 	<Option Name="AlternateGauntletDevicePool" DefaultValue="UE5" Description="Alternate Gauntlet Device Pool to use to run target platform tests"/>
 	<Option Name="AlternateDevicePoolPlatforms" DefaultValue="XSX;PS5" Description="Platforms that need to use the Alternate Gauntlet Device Pool"/>
+	<Option Name="ExtraGauntletCommonArgs" DefaultValue="" Description="Additional Common arguments passed to Gauntlet command line"/>
+	<Option Name="ExtraGauntletTargetTestArgs" DefaultValue="" Description="Additional Target test arguments passed to Gauntlet command line"/>
+	<Option Name="ExtraGauntletEditorTestArgs" DefaultValue="" Description="Additional Editor test arguments passed to Gauntlet command line"/>
 	
-	<Property Name="GauntletCommonArgs" Value="-maxduration=$(GauntletTimeOut) -unattended -verbose -ResumeOnCriticalFailure" />
+	<Property Name="GauntletCommonArgs" Value="-maxduration=$(GauntletTimeOut) -unattended -verbose -ResumeOnCriticalFailure $(ExtraGauntletCommonArgs)" />
 
 	<Property Name="GauntletDeviceArgs" Value="" />
 	<Property Name="GauntletDeviceArgs" Value="-deviceurl=https://horde.devtools.epicgames.com/ -reboot -cleardevices -sandbox=&quot;$(ProjectName)&quot;" If="$(IsBuildMachine)" />
@@ -29,9 +32,9 @@
 	<Property Name="GauntletDeviceArgs" Value="$(GauntletDeviceArgs) -devicepool=&quot;$(GauntletDevicePool)&quot;" If="$(IsBuildMachine)"/>
 	<Property Name="GauntletDeviceAltArgs" Value="$(GauntletDeviceAltArgs) -devicepool=&quot;$(AlternateGauntletDevicePool)&quot;" If="$(IsBuildMachine)"/>
 
-	<Property Name="GauntletTargetTestArgs" Value="$(GauntletCommonArgs) $(GauntletDeviceArgs)" />
-	<Property Name="GauntletTargetTestAltArgs" Value="$(GauntletCommonArgs) $(GauntletDeviceAltArgs)" />
-	<Property Name="GauntletEditorTestArgs" Value="-build=editor $(GauntletCommonArgs)" />
+	<Property Name="GauntletTargetTestArgs" Value="$(GauntletCommonArgs) $(GauntletDeviceArgs) $(ExtraGauntletTargetTestArgs)" />
+	<Property Name="GauntletTargetTestAltArgs" Value="$(GauntletCommonArgs) $(GauntletDeviceAltArgs) $(ExtraGauntletTargetTestArgs)" />
+	<Property Name="GauntletEditorTestArgs" Value="-build=editor $(GauntletCommonArgs) $(ExtraGauntletEditorTestArgs)" />
 
     <!-- this is fully generated dynamiclly when needed, it's simply here so it can be used without worrying if its set -->
     <Property Name="GauntletReportArgs" Value="" />
