Remove heuristic to guess TargetName. Introduce IsContenOnlyProject to force UnrealGame as TargetName

#rnx

diff --git a/Engine/Build/Graph/Tasks/BuildAndTestProject.xml b/Engine/Build/Graph/Tasks/BuildAndTestProject.xml
--- a/Engine/Build/Graph/Tasks/BuildAndTestProject.xml
+++ b/Engine/Build/Graph/Tasks/BuildAndTestProject.xml
@@ -6,6 +6,8 @@
 	<Option Name="ProjectName" DefaultValue="None" Restrict=".+" Description="Name of Project"/>
 	<Option Name="ProjectPath" DefaultValue="None" Restrict=".+" Description="Path to the folder that contains your project"/>
 
+	<Option Name="IsContentOnlyProject" DefaultValue="false" Restrict="true|false" Description="Is the project a content only project?"/>	
+
 	<Option Name="TargetName" DefaultValue="" Restrict="" Description="If your project supports both Game and Client you can set this, otherwise one will be selected"/>
 	<Option Name="EditorTargetName" DefaultValue="" Restrict="" Description="Specify the editor target name if different than (ProjectName)Editor"/>
 
@@ -100,23 +102,10 @@
 	<Property Name="TargetProject" Value="$(ProjectPath)/$(ProjectName).uproject" />
 	<Error Message="Project file $(TargetProject) does not exist" If="!Exists('$(TargetProject)')" />
 
-	<!-- If the user didn't supply a TargetName then try to auto detect it. This is reliable uness the project has 'Game' in it but the target doesn't -->
-	<Do If="'$(TargetName)' == ''">
-		<!-- 
-			If the user didn't supply a TargetName, try to figure it out based on GameName.Target or GameNameClient.target.
-			Note - if the project is ShooterGame and the target is ShooterClient.txt then this will fail and we'll tell the user 
-		-->
-		<Property Name="TargetName" Value="$(ProjectName)Client" If="Exists('$(ProjectPath)/Source/$(ProjectName)Client.Target.cs')" />
-		<Property Name="TargetName" Value="$(ProjectName)" If="Exists('$(ProjectPath)/Source/$(ProjectName).Target.cs')" />
-		
-		<Do If="'$(TargetName)' == ''">
-			<Do If="Exists('$(ProjectPath)/Source/$(ProjectName)/$(ProjectName).Build.cs')">
-				<Error Message="The project has source but the name the target couldn't be determined. use -set:TargetName:Foo to ensure things build correctly" />
-			</Do>
-			<Property Name="TargetName" Value="UnrealGame" />
-			<Property Name="EditorTargetName" Value="UnrealEditor" If="'$(EditorTargetName)' == ''"/>
-		</Do>
-	</Do>
+	<!-- If the user didn't supply a TargetName then use the ProjectName -->
+	<Property Name="TargetName" Value="UnrealGame" If="$(IsContentOnlyProject) And '$(TargetName)' == ''"/>
+	<Property Name="TargetName" Value="$(ProjectName)" If="'$(TargetName)' == ''"/>
+	<Property Name="EditorTargetName" Value="UnrealEditor" If="$(IsContentOnlyProject) And '$(EditorTargetName)' == ''"/>
 	<Property Name="EditorTargetName" Value="$(ProjectName)Editor" If="'$(EditorTargetName)' == ''"/>
 
 	<!-- If skip test is true, add all target platforms to the "do not test" list -->
