Add another platform

#rnx

diff --git a/Engine/Build/Graph/Tasks/Inc/PlatformOptions.xml b/Engine/Build/Graph/Tasks/Inc/PlatformOptions.xml
--- a/Engine/Build/Graph/Tasks/Inc/PlatformOptions.xml
+++ b/Engine/Build/Graph/Tasks/Inc/PlatformOptions.xml
@@ -27,7 +27,7 @@
 <Property Name="ShowExplicitFlags" Value="true" />
 
 <Property Name="AllEditorPlatforms" Value="Mac+Win64" />
-<Property Name="AllTargetPlatforms" Value="Android+HoloLens+iOS+Mac+PS4+PS5+Stadia+Switch+tvOS+Win64+WinGDK+XboxOneGDK+XSX" />
+<Property Name="AllTargetPlatforms" Value="Android+HoloLens+iOS+Mac+PS4+PS5+Stadia+Switch+tvOS+Win64+WinGDK+XboxOneGDK+XB1+XSX" />
 
 <Option Name="EditorPlatforms" DefaultValue="" Description="Alternative to 'WithFooEditor', a list of editors to build and test. Win64+Mac etc"/>
 <Option Name="TargetPlatforms" DefaultValue="" Description="Alternative to 'WithFoo', a list of Platforms to build and test. PS4+XSX etc"/>
@@ -61,6 +61,7 @@
 	<Option Name="WithWin64" DefaultValue="false" Description="Build and test Win64"/>
 	<Option Name="WithWinGDK" DefaultValue="false" Description="Build and test WinGDK"/>
 	<Option Name="WithXboxOneGDK" DefaultValue="false" Description="Build and test XboxOneGDK"/>
+	<Option Name="WithXB1" DefaultValue="false" Description="Build and test XB1"/>
 	<Option Name="WithXSX" DefaultValue="false" Description="Build and test XSX"/>
 
 	<Do If="'$(EditorPlatforms)' == ''" >
@@ -85,6 +86,7 @@
 		<Property Name="TargetPlatforms" Value="$(TargetPlatforms)Win64+" If="$(WithWin64)" />
 		<Property Name="TargetPlatforms" Value="$(TargetPlatforms)WinGDK+" If="$(WithWinGDK)" />
 		<Property Name="TargetPlatforms" Value="$(TargetPlatforms)XboxOneGDK+" If="$(WithXboxOneGDK)" />
+		<Property Name="TargetPlatforms" Value="$(TargetPlatforms)XB1+" If="$(WithXB1)" />
 		<Property Name="TargetPlatforms" Value="$(TargetPlatforms)XSX+" If="$(WithXSX)" />
 	</Do>
 </Do>
