Updated the HLOD generation BuildGraph example script to include most changes from our internal demo

diff --git a/Engine/Build/Graph/Examples/BuildWorldPartitionHLODs.xml b/Engine/Build/Graph/Examples/BuildWorldPartitionHLODs.xml
--- a/Engine/Build/Graph/Examples/BuildWorldPartitionHLODs.xml
+++ b/Engine/Build/Graph/Examples/BuildWorldPartitionHLODs.xml
@@ -1,22 +1,40 @@
 <?xml version='1.0' ?>
 <BuildGraph xmlns="http://www.epicgames.com/BuildGraph" xmlns:xsi="http://www.w3.org/2001/XMLSchema-instance" xsi:schemaLocation="http://www.epicgames.com/BuildGraph ../Schema.xsd" >
 
-	<Option Name="MapName" DefaultValue="" Description="The map for which HLODs will be generated."/>
+	<!-- Options -->
 	<Option Name="ProjectName" DefaultValue="" Description="Name of the project."/>
 	<Option Name="EditorTarget" DefaultValue="" Description="Name of the editor target to use."/>
+	<Option Name="MapName" DefaultValue="" Description="The map for which HLODs will be generated."/>
 	<Option Name="BuilderCount" DefaultValue="10" Description="Number of machines to use"/>
+	<Option Name="SubmitHLODs" Restrict="true|false" DefaultValue="false" Description="Submit generated HLODs"/>
+	<Option Name="ShelveHLODs" Restrict="true|false" DefaultValue="false" Description="Shelve generated HLODs"/>
+	<Option Name="ShelveUser" DefaultValue="" Description="Shelve result to the provided user workspace"/>
+	<Option Name="ShelveWorkspace" DefaultValue="" Description="Shelve result to the provided user workspace"/>
+	
+	<Error Message="Missing map name. You must provide a map name using -MapName=..." If="'$(MapName)' == ''"/>
+	<Error Message="Missing project name. You must provide the project name using -ProjectName=..." If="'$(ProjectName)' == ''"/>
+	<Error Message="Missing editor target. You must provide the target name using -EditorTarget=..." If="'$(EditorTarget)' == ''"/>
+	<Error Message="Conflicting options found: -SubmitHLODs and -ShelveHLODs" If="$(SubmitHLODs) And $(ShelveHLODs)"/>
+	<Error Message="Cannot shelve result without providing a proper P4 user and workspace" If="$(ShelveHLODs) And ('$(ShelveUser)' == '' Or '$(ShelveWorkspace)' == '')"/>
+	<Error Message="Provided a shelve user and workspace with the -SubmitHLODs option, are you really sure you want to submit ? Aborting." If="$(SubmitHLODs) And ('$(ShelveUser)' != '' Or '$(ShelveWorkspace)' != '')"/>
+	<Error Message="Trying to build HLODs on a buildmachine, but without any output (missing -SubmitHLODs or -ShelveHLODs)" If="$(IsBuildMachine) And !$(SubmitHLODs) And !$(ShelveHLODs)"/>
 
 	<!-- Perforce environment variables -->
 	<EnvVar Name="uebp_PORT"/>
 	<EnvVar Name="uebp_USER"/>
 	<EnvVar Name="uebp_CLIENT"/>
+
+	<Property Name="HLODBuilder" Value="WorldPartitionHLODsBuilder"/>
 		
-	<Property Name="HLODCommonBuilderArgs" Value="$(MapName) -Builder=WorldPartitionHLODsBuilder -DistributedBuild"/>
+	<Property Name="HLODCommonBuilderArgs" Value="$(MapName) -DistributedBuild"/>
 	<Property Name="HLODCommonBuilderArgs" Value="-unattended -buildmachine $(HLODCommonBuilderArgs)" If="$(IsBuildMachine)"/>
 	
 	<Property Name="HLODBuilderSCC_None" Value="-SCCProvider=None"/>
-	<Property Name="HLODBuilderSCC_P4" Value="-SCCProvider=Perforce -P4Port=$(uebp_PORT) -P4User=$(uebp_USER) -P4Client=$(uebp_CLIENT)"/>
-		
+	<Property Name="HLODBuilderSCC_P4" Value="-SCCProvider=Perforce"/>
+	<Property Name="HLODBuilderSCC_P4" Value="$(HLODBuilderSCC_P4) -P4Port=$(uebp_PORT) -P4User=$(uebp_USER) -P4Client=$(uebp_CLIENT)" If="$(IsBuildMachine)"/>
+
+
+	<!-- Prerequisites -->
 	<Agent Name="HLOD Generation Prerequisites" Type="Win64">
 		<Node Name="Update Version Files">
 			<SetVersion Change="$(Change)" Branch="$(EscapedBranch)" If="$(IsBuildMachine)"/>
@@ -27,22 +45,25 @@
 			<Compile Target="ShaderCompileWorker" Platform="Win64" Configuration="Development"/>
 		</Node>
 		
-		<Node Name="Compile $(EditorTarget) Win64">
+		<Node Name="Compile Editor Win64">
 			<Compile Target="$(EditorTarget)" Platform="Win64" Configuration="Development"/>
 		</Node>
 	</Agent>
-	<Property Name="HLODSetupDependencies" Value="Compile $(EditorTarget) Win64;Compile Tools Win64"/>
-	
+	<Property Name="HLODSetupDependencies" Value="Compile Editor Win64;Compile Tools Win64"/>
+
+
+	<!-- Setup -->
 	<Agent Name="HLOD Setup" Type="Win64">
 		<Node Name="HLOD Setup" Requires="$(HLODSetupDependencies)">
-			<Commandlet Name="WorldPartitionBuilderCommandlet" Project="$(ProjectName)" Arguments="$(HLODCommonBuilderArgs) $(HLODBuilderSCC_None) -AllowCommandletRendering -SetupHLODs -BuilderCount=$(BuilderCount)"/>
+			<Command Name="WorldPartitionBuilder" Arguments="-Project=$(ProjectName) -Builder=$(HLODBuilder) -CommandletArgs=&quot;$(HLODCommonBuilderArgs) $(HLODBuilderSCC_None) -SetupHLODs -BuilderCount=$(BuilderCount)&quot;"/>
 			<GatherBuildProductsFromFile BuildProductsFile="$(RootDir)/HLODTemp/BuildProducts.txt"/>
 		</Node>
 	</Agent>
 
+
+	<!-- Build -->
 	<Property Name="HLODBuildNodes" Value=""/>
 	<Property Name="HLODBuildCompleteNodes" Value=""/>
-
 	<ForEach Name="HLODBuildJobIndex" Values="0;1;2;3;4;5;6;7;8;9;10;11;12;13;14;15;16;17;18;19;20;21;22;23;24;25;26;27;28;29;30;31;32;33;34;35;36;37;38;39;40;41;42;43;44;45;46;47;48;49">
 		<Do If="'$(HLODBuildJobIndex)' &lt; '$(BuilderCount)'">
 			<Property Name="HLODBuildAgent" Value="HLOD Build Agent $(HLODBuildJobIndex)"/>
@@ -53,7 +74,7 @@
 			
 			<Agent Name="$(HLODBuildAgent)" Type="Win64">
 				<Node Name="$(HLODBuildNode)" Requires="HLOD Setup">
-					<Commandlet Name="WorldPartitionBuilderCommandlet" Project="$(ProjectName)" Arguments="$(HLODCommonBuilderArgs) $(HLODBuilderSCC_None) -AllowCommandletRendering -BuildHLODs -BuilderIdx=$(HLODBuildJobIndex)"/>
+					<Command Name="WorldPartitionBuilder" Arguments="-Project=$(ProjectName) -Builder=$(HLODBuilder) -CommandletArgs=&quot;$(HLODCommonBuilderArgs) $(HLODBuilderSCC_None) -AllowCommandletRendering -BuildHLODs -BuilderIdx=$(HLODBuildJobIndex)&quot;"/>
 				</Node>
 				<Node Name="$(HLODBuildCompleteNode)" After="$(HLODBuildNode)">
 					<GatherBuildProductsFromFile BuildProductsFile="$(RootDir)/HLODTemp/BuildProducts.txt"/>
@@ -61,13 +82,18 @@
 			</Agent>
 		</Do>
 	</ForEach>
-	
-	<Agent Name="HLOD Submit" Type="Win64_NonAWS">
-		<Node Name="HLOD Submit" Requires="$(HLODBuildNodes);$(HLODBuildCompleteNodes);HLOD Setup">
-			<Commandlet Name="WorldPartitionBuilderCommandlet" Project="$(ProjectName)" Arguments="$(HLODCommonBuilderArgs) $(HLODBuilderSCC_P4) -SubmitHLODs"/>
+
+
+	<!-- Finalize -->
+	<Property Name="SCCOperationArgs" Value=""/>
+	<Property Name="SCCOperationArgs" Value="-Submit" If="$(SubmitHLODs)"/>
+	<Property Name="SCCOperationArgs" Value="-ShelveUser=$(ShelveUser) -ShelveWorkspace=$(ShelveWorkspace)" If="$(ShelveHLODs) and '$(ShelveUser)' != '' and '$(ShelveWorkspace)' != ''"/>
+	<Agent Name="HLOD Finalize" Type="Win64">
+		<Node Name="HLOD Finalize" Requires="$(HLODBuildNodes);$(HLODBuildCompleteNodes);HLOD Setup">
+			<Command Name="WorldPartitionBuilder" Arguments="-Project=$(ProjectName) -Builder=$(HLODBuilder) -P4 $(SCCOperationArgs) -CommandletArgs=&quot;$(HLODCommonBuilderArgs) $(HLODBuilderSCC_P4) -FinalizeHLODs&quot;"/>
 		</Node>
 	</Agent>
 
-	<Aggregate Name="HLOD Generation" Requires="HLOD Submit"/>
+	<Aggregate Name="HLOD Generation" Requires="HLOD Finalize"/>
 	
 </BuildGraph>
\ No newline at end of file
