Add the ability to specify additional tools in the BuildEditorAndTools build graph

diff --git a/Engine/Build/Graph/Examples/BuildEditorAndTools.xml b/Engine/Build/Graph/Examples/BuildEditorAndTools.xml
--- a/Engine/Build/Graph/Examples/BuildEditorAndTools.xml
+++ b/Engine/Build/Graph/Examples/BuildEditorAndTools.xml
@@ -25,6 +25,7 @@
 			-set:EditorTarget=<TargetName>		Specifies the editor target to be built (defaults to UnrealEditor)
 			-set:UProjectPath=<ProjectPath>		Specifies the uproject path to pass as -Project in the compile args (defaults to none)
 			-set:GameTargets=<TargetName>		Specifies the game targets to be built (defaults to none)
+			-set:AdditionalTools=<A>;<B>;<C>..	Specifies the additional tool targets to be built (defaults to none)
 			-set:TargetPlatforms=<A>;<B>;<C>... Specifies the target platforms to build the game for (defaults to Win64)
 			-set:OutputDir=<Path>				Specifies the output directory for the compiled binaries
 			-set:Licensee=true/false			Whether to mark the changelist number as being from a licensee Perforce server (this prevents conflicts with Epic CL's being a higher number)
@@ -36,6 +37,7 @@
 	<Option Name="EditorTarget" Restrict="[^ ]+" DefaultValue="UnrealEditor" Description="Name of the editor target to be built"/>
 	<Option Name="UProjectPath" Restrict="[^ ]*" DefaultValue="" Description="Optional Uproject path to include in the arguments"/>
 	<Option Name="GameTargets" Restrict="[^ ]*" DefaultValue="" Description="Names of game targets to build separated by semicolons, eg. UnrealClient;UnrealServer"/>
+	<Option Name="AdditionalTools" Restrict="[^ ]*" DefaultValue="" Description="Tool Compiles to build seperated by semicolons, eg. UnrealMultiUserServer;BuildPatchTool"/>
 	<Option Name="TargetPlatforms" Restrict="[^ ]*" DefaultValue="Win64" Description="List of the target platforms to build for, separated by semicolons, eg. Win64;Android"/>
 	<Option Name="OutputDir" DefaultValue="$(RootDir)\LocalBuilds\$(EditorTarget)Binaries" Description ="Output directory for compiled binaries"/>
 	<Option Name="Licensee" Restrict="true|false" DefaultValue="true" Description="Whether to treat the changelist number as being from a licensee Perforce server"/>
@@ -60,6 +62,9 @@
 			<Compile Target="ShaderCompileWorker" Platform="Win64" Configuration="Development" Tag="#ToolBinaries"/>
 			<Compile Target="UnrealLightmass" Platform="Win64" Configuration="Development" Tag="#ToolBinaries"/>
 			<Compile Target="UnrealPak" Configuration="Development" Platform="Win64" Tag="#ToolBinaries"/>
+			<ForEach Name="AdditonalTool" Values="$(AdditionalTools)">
+				<Compile Target="$(AdditonalTool)" Platform="Win64" Configuration="Development" Tag="#ToolBinaries"/>
+			</ForEach>
 			
 			<!-- 
 				This exe is a copy of ShaderCompileWorker.exe, created as a post-build step. See \Engine\Source\Programs\ShaderCompileWorker\ShaderCompileWorker.Target.cs.
