Adding an option to specify the archive name of a PCBs archive

Allows building a project that does not specifies an editor, using UnrealEditor, but is its own project to prevent archives stomping over each other

#rnx

diff --git a/Engine/Build/Graph/Examples/BuildEditorAndTools.xml b/Engine/Build/Graph/Examples/BuildEditorAndTools.xml
--- a/Engine/Build/Graph/Examples/BuildEditorAndTools.xml
+++ b/Engine/Build/Graph/Examples/BuildEditorAndTools.xml
@@ -31,6 +31,7 @@
 			-set:Licensee=true/false			Whether to mark the changelist number as being from a licensee Perforce server (this prevents conflicts with Epic CL's being a higher number)
 			-set:Versioned=true/false			Whether to embed version information in the compiled binaries. This is true by default for build machines.
 			-set:ArchiveStream=<Path>			Stream that will contain zip file of the compiled binaries, for use by UnrealGameSync.
+			-set:ArchiveName=<ArchiveName>		Names of archive to generate instead of using the GameTarget, when the GameTarget may be UnrealEditor and result in a clash.
 			-set:ForceSubmit=true/false			Whether to force submit to happen even if another change has been submitted (resolves in favor of local files)
 	-->
 
@@ -44,6 +45,7 @@
 	<Option Name="Versioned" Restrict="true|false" DefaultValue="$(IsBuildMachine)" Description="Whether to embed changelist number into binaries"/>
 	<Option Name="ArchiveStream" Restrict="(?://.*)?" DefaultValue="" Description="Stream path that will contain zip file of the compiled binaries, for use by UnrealGameSync"/>
 	<Option Name="ForceSubmit" Restrict="true|false" DefaultValue="false" Description="Forces the submit to happen even if another change has been submitted (resolves in favor of local files)"/>
+	<Option Name="ArchiveName" Restrict="[^ ]*" DefaultValue="$(EditorTarget)" Description="Names of archive to generate instead of using the GameTarget, when the GameTarget may be UnrealEditor and result in a clash"/>
 	<Option Name="SymbolStorePath" DefaultValue="" Description ="Output directory for symbol storage"/>
 
 	<EnvVar Name="COMPUTERNAME"/>
@@ -152,7 +154,7 @@
 
 			<!-- Create the zip file and submit it to Perforce -->
 			<Property Name="ArchivePerforceDir" Value="$(ArchiveDir)\Perforce"/>
-			<Property Name="ArchiveFile" Value="$(ArchivePerforceDir)\$(EscapedBranch)-$(EditorTarget).zip"/>
+			<Property Name="ArchiveFile" Value="$(ArchivePerforceDir)\$(EscapedBranch)-$(ArchiveName).zip"/>
 			<Zip FromDir="$(ArchiveStagingDir)" ZipFile="$(ArchiveFile)"/>
 			<Warning Message="The 'ArchiveStream' argument is not set. Binaries will not be submitted." If="'$(ArchiveStream)' == ''"/>
 			<Property Name="SubmitClient" Value="$(COMPUTERNAME)_ArchiveForUGS"/>
