add another platform to more build graph nodes.

#rnx

diff --git a/Engine/Build/Graph/Tasks/BuildAndTestProject.xml b/Engine/Build/Graph/Tasks/BuildAndTestProject.xml
--- a/Engine/Build/Graph/Tasks/BuildAndTestProject.xml
+++ b/Engine/Build/Graph/Tasks/BuildAndTestProject.xml
@@ -170,7 +170,7 @@
 				<Do If="$(WithRemoteExecWorkers)" >
 					<Compile Target="BaseTextureBuildWorker" Platform="$(EditorPlatform)" Configuration="Development" Arguments="-project=&quot;$(RootDir)/$(TargetProject)&quot; $(GenericCompileArguments)"/>
 					<Do If="'$(EditorPlatform)' == 'Win64'" >
-						<ForEach Name="TargetPlatform" Values="PS4+PS5+Switch+XboxOneGDK+XSX" Separator="+">
+						<ForEach Name="TargetPlatform" Values="PS4+PS5+Switch+XboxOneGDK+XB1+XSX" Separator="+">
 							<Compile Target="$(TargetPlatform)TextureBuildWorker" Platform="$(EditorPlatform)" Configuration="Development" Arguments="-project=&quot;$(RootDir)/$(TargetProject)&quot; $(GenericCompileArguments)" If="ContainsItem('$(TargetPlatforms)', '$(TargetPlatform)', '+')"/>
 						</ForEach>
 					</Do>
