fix staging content only project

diff --git a/Engine/Build/Graph/Tasks/BuildAndTestProject.xml b/Engine/Build/Graph/Tasks/BuildAndTestProject.xml
--- a/Engine/Build/Graph/Tasks/BuildAndTestProject.xml
+++ b/Engine/Build/Graph/Tasks/BuildAndTestProject.xml
@@ -263,7 +263,7 @@
 		<Property Name="PlatformEditorNodeName" Value="Compile $(ProjectName)Editor $(HostAgentType)" />
 		<!-- Define these names up front as we use some of them several times -->
 		<Property Name="PlatformAgentName" Value="$(DisplayTargetName) $(TargetPlatform)" />
-		<Property Name="PlatformCompileNodeName" Value="Compile $(TargetName) $(TargetPlatform)" />
+		<Property Name="PlatformCompileNodeName" Value="Compile $(DisplayTargetName) $(TargetPlatform)" />
 		<Property Name="UploadSymbolsNodeName" Value="" />
 		<Property Name="UploadSymbolsNodeName" Value="Upload Symbols $(TargetName) $(TargetPlatform)" If="'$(SymbolServerPath)' != ''" />
 		<Property Name="PlatformCookNodeName" Value="Cook $(DisplayTargetName) $(TargetPlatform)" />
@@ -289,7 +289,7 @@
 		<Property Name="DoUnsyncHashing" Value="$(HostAgentType) == 'Win64'" />
 
 		<!-- Default BuildCookRun args needed when staging/packaging/archiving this configuration -->
-		<Property Name="BCRArgs" Value="--nocompile -project=$(TargetProject) -platform=$(TargetPlatform) -cookflavor=&quot;$(BCRCookFlavor)&quot; -configuration=$(PackageConfigurations)" />
+		<Property Name="BCRArgs" Value="-project=$(TargetProject) -platform=$(TargetPlatform) -cookflavor=&quot;$(BCRCookFlavor)&quot; -configuration=$(PackageConfigurations)" />
 		<Property Name="BCRArgs" Value="$(BCRArgs) -client" If="'$(TargetType)' == 'Client'" />
 		<Property Name="BCRArgs" Value="$(BCRArgs) -server" If="'$(DedicatedServer)'" />
 
@@ -301,13 +301,17 @@
 		<Agent Name="Compile $(PlatformAgentName)" Type="$(CompileAgentType)">
 
 			<!-- Compile -->
-			<Node Name="$(PlatformCompileNodeName)" Requires="$(PlatformToolsNodeName)">
+			<Node Name="$(PlatformCompileNodeName)" Requires="$(PlatformToolsNodeName)" Produces="#$(PlatformCompileNodeName) Binaries">
 				<ForEach Name="TargetConfiguration" Values="$(TargetConfigurations)" Separator="+">
 					<Compile Target="$(TargetName)" Platform="$(TargetPlatform)" Configuration="$(TargetConfiguration)" Arguments="-BuildVersion=&quot;$(BuildName)&quot; $(GenericCompileArguments) $(ExtraTargetCompileArguments)"/>
 					<Do If="'$(DedicatedServer)'" >
 						<Compile Target="$(TargetName)Server" Platform="Win64" Configuration="$(TargetConfiguration)" Arguments="-BuildVersion=&quot;$(BuildName)&quot; $(GenericCompileArguments) $(ExtraTargetCompileArguments)"/>
 					</Do>
 				</ForEach>
+				<Do If="$(IsContentOnlyProject)">
+					<Command Name="BuildGame" Arguments="$(BCRArgs) -notools -ubtargs=&quot;-nobuilduht -BuildVersion=$(BuildName) $(GenericCompileArguments) $(ExtraTargetCompileArguments)&quot;" />
+					<Tag Files="$(RootDir)/$(ProjectPath)/Binaries/..." With="#$(PlatformCompileNodeName) Binaries"/>
+				</Do>
 			</Node>
 
 			<!-- (Optional) Upload Symbols -->
@@ -331,7 +335,7 @@
 			</Node>
 			
 			<!-- Stage -->
-			<Node Name="$(PlatformStageNodeName)" Requires="$(PlatformCompileNodeName);#$(PlatformCookNodeName) Complete">
+			<Node Name="$(PlatformStageNodeName)" Requires="$(PlatformCompileNodeName);#$(PlatformCompileNodeName) Binaries;#$(PlatformCookNodeName) Complete">
 				<Command Name="BuildCookRun" Arguments="$(BCRArgs) $(StageAndPackageArguments) -skipbuild -skipcook -stage -pak" />
 			</Node>
 			
