Allow to override the agent type for BuildEditorAndTools job
#rnx

diff --git a/Engine/Build/Graph/Examples/BuildEditorAndTools.xml b/Engine/Build/Graph/Examples/BuildEditorAndTools.xml
--- a/Engine/Build/Graph/Examples/BuildEditorAndTools.xml
+++ b/Engine/Build/Graph/Examples/BuildEditorAndTools.xml
@@ -33,6 +33,7 @@
 			-set:ArchiveStream=<Path>			Stream that will contain zip file of the compiled binaries, for use by UnrealGameSync.
 			-set:ArchiveName=<ArchiveName>		Names of archive to generate instead of using the GameTarget, when the GameTarget may be UnrealEditor and result in a clash.
 			-set:ForceSubmit=true/false			Whether to force submit to happen even if another change has been submitted (resolves in favor of local files)
+			-set:PreferredAgent=<AgentName>		Specifies the agent to use for job, defaults to CompileWin64
 	-->
 
 	<Option Name="EditorTarget" Restrict="[^ ]+" DefaultValue="UnrealEditor" Description="Name of the editor target to be built"/>
@@ -45,13 +46,14 @@
 	<Option Name="Versioned" Restrict="true|false" DefaultValue="$(IsBuildMachine)" Description="Whether to embed changelist number into binaries"/>
 	<Option Name="ArchiveStream" Restrict="(?://.*)?" DefaultValue="" Description="Stream path that will contain zip file of the compiled binaries, for use by UnrealGameSync"/>
 	<Option Name="ForceSubmit" Restrict="true|false" DefaultValue="false" Description="Forces the submit to happen even if another change has been submitted (resolves in favor of local files)"/>
+	<Option Name="PreferredAgent" DefaultValue="CompileWin64" Description ="Agent type to use for job"/>
 	<Option Name="ArchiveName" Restrict="[^ ]*" DefaultValue="$(EditorTarget)" Description="Names of archive to generate instead of using the GameTarget, when the GameTarget may be UnrealEditor and result in a clash"/>
 	<Option Name="SymbolStorePath" DefaultValue="" Description ="Output directory for symbol storage"/>
 
 	<EnvVar Name="COMPUTERNAME"/>
 	<EnvVar Name="P4CLIENT"/>
 	
-	<Agent Name="Default Agent" Type="CompileWin64">
+	<Agent Name="Default Agent" Type="$(PreferredAgent)">
 	
 		<!-- Update the engine version files -->
 		<Node Name="Update Version Files">
