Clean up in BuildlAndTetstProject buildgraph
* Remove repetition of the project name in node name
* Fix labels to use the one from the Horde Dashboard
* Avoid using exclude property in Label

#rnx

diff --git a/Engine/Build/Graph/Tasks/BuildAndTestProject.xml b/Engine/Build/Graph/Tasks/BuildAndTestProject.xml
--- a/Engine/Build/Graph/Tasks/BuildAndTestProject.xml
+++ b/Engine/Build/Graph/Tasks/BuildAndTestProject.xml
@@ -161,7 +161,7 @@
 				<SetVersion Change="$(Change)" Branch="$(EscapedBranch)" If="$(Versioned) and $(IsBuildMachine)"/>
 			</Node>
 			<!-- First generate nodes to build the required tools (e.g. UHT, ShaderCompiler, Editor -->
-			<Node Name="Compile $(ProjectName)Tools $(EditorPlatform)" Requires="Update Version Files $(EditorPlatform)">
+			<Node Name="Compile Tools $(EditorPlatform)" Requires="Update Version Files $(EditorPlatform)">
 				<Compile Target="UnrealHeaderTool" Platform="$(EditorPlatform)" Configuration="Development" Arguments="$(GenericCompileArguments)"/>
 				<Compile Target="ShaderCompileWorker" Platform="$(EditorPlatform)" Configuration="Development" Arguments="$(GenericCompileArguments)"/>
 				<Compile Target="UnrealPak" Platform="$(EditorPlatform)" Configuration="Development" Arguments="$(GenericCompileArguments)"/>
@@ -176,17 +176,17 @@
 					</Do>
 				</Do>
 			</Node>
-			<Node Name="Compile $(ProjectName)Editor $(EditorPlatform)" Requires="Compile $(ProjectName)Tools $(EditorPlatform)">
+			<Node Name="Compile $(ProjectName)Editor $(EditorPlatform)" Requires="Compile Tools $(EditorPlatform)">
 				<Compile Target="$(EditorTargetName)" Platform="$(EditorPlatform)" Configuration="Development" Arguments="$(GenericCompileArguments) $(ExtraEditorCompileArguments)"/>
 			</Node>
 			<!--Add these nodes to our dependency list -->
-			<Property Name="GeneratedToolNodes" Value="Compile $(ProjectName)Editor $(EditorPlatform); Compile $(ProjectName)Tools $(EditorPlatform)" If="!$(SkipBuild)" />
+			<Property Name="GeneratedToolNodes" Value="Compile $(ProjectName)Editor $(EditorPlatform); Compile Tools $(EditorPlatform)" If="!$(SkipBuild)" />
 		</Agent>
 		<Agent Name="Test $(ProjectName)Editor $(EditorPlatform)" Type="$(TestAgentType)">
 			<!-- Prerequisites for the editor test to run -->
 			<Property Name="TestPrerequisites" Value="$(GeneratedToolNodes)"/>
 			<!-- Declare a boot test for the editor-->
-			<Property Name="PlatformBootTestNodeName" Value="$(ProjectName)Editor $(EditorPlatform) Test=BootTest" />
+			<Property Name="PlatformBootTestNodeName" Value="Editor $(EditorPlatform) Test=BootTest" />
 			<Do If="'$(AutomationReportOutputDirectory)' != ''">
 				<Property Name="OutputFragmentPath" Value="$(EditorPlatform)Editor/UE.EditorBootTest" />
 				<Property Name="OutputReportPath" Value="$(AutomationReportOutputDirectory)/$(OutputFragmentPath)" />
@@ -196,7 +196,7 @@
 			<Node Name="$(PlatformBootTestNodeName)" Requires="$(GeneratedToolNodes)">
 				<Command Name="RunUnreal" Arguments="-project=$(TargetProject) -build=Editor -test=&quot;UE.EditorBootTest&quot; $(GauntletEditorTestArgs) $(GauntletReportArgs)" />
 			</Node>
-			<Property Name="ProjectEditorDDC" Value="$(ProjectName)Editor DDC $(EditorPlatform)" />
+			<Property Name="ProjectEditorDDC" Value="Editor DDC $(EditorPlatform)" />
 			<Node Name="$(ProjectEditorDDC)" Requires="$(GeneratedToolNodes)">
 				<Property Name="DDCTargetPlatform" Value="$(EditorPlatform)" />
 				<Property Name="DDCTargetPlatform" Value="Windows" If="'$(DDCTargetPlatform)' == 'Win64'"/>
@@ -214,7 +214,7 @@
 					<Property Name="GauntletReportArgs" Value="-logdir=&quot;$(OutputReportPath)&quot; -artifactname=Gauntlet -ReportExportPath=&quot;$(OutputReportPath)&quot;" />
 					<Property Name="GauntletReportArgs" Value="$(GauntletReportArgs) -ReportURL=&quot;$(AutomationReportUrl)/$(OutputFragmentPath)/&quot;" If="'$(AutomationReportUrl)' != ''" />
 				</Do>
-				<Property Name="TestNodeName" Value="$(ProjectName)Editor $(EditorPlatform) Test=$(TestName)" />
+				<Property Name="TestNodeName" Value="Editor $(EditorPlatform) Test=$(TestName)" />
 				<Node Name="$(TestNodeName)" Requires="$(TestPrerequisites)">
 					<Command Name="RunUnreal" Arguments="-project=$(TargetProject) -build=Editor -test=&quot;$(TestName)&quot; $(GauntletReportArgs) $(GauntletEditorTestArgs)" />
 				</Node>
@@ -228,8 +228,8 @@
 
 			<!-- If this editor was requested (as opposed to being required due to target platforms) then link in the test nodes and set up labels -->
 			<Do If="ContainsItem('$(EditorPlatforms)', '$(EditorPlatform)', '+') and !ContainsItem('$(EditorPlatformsToSkipTests)', '$(EditorPlatform)', '+')">
-				<Label Category="Test" Name="BootTest $(ProjectName)Editor $(EditorPlatform)" Requires="$(PlatformBootTestNodeName)" />
-				<Label Category="Test" Name="Test $(ProjectName)Editor $(EditorPlatform)" Requires="$(PlatformEditorTestNodes)" />
+				<Label Category="Tests" Name="BootTest $(ProjectName)Editor $(EditorPlatform)" Requires="$(PlatformBootTestNodeName)" />
+				<Label Category="Tests" Name="Test $(ProjectName)Editor $(EditorPlatform)" Requires="$(PlatformEditorTestNodes)" />
 				<Property Name="GeneratedEditorTestNodes" Value="$(GeneratedEditorTestNodes);$(PlatformBootTestNodeName)" If="$(EditorBootTest)" />
 				<Property Name="GeneratedEditorTestNodes" Value="$(GeneratedEditorTestNodes);$(ProjectEditorDDC)" If="$(FillEditorDDC)" />
 				<Property Name="GeneratedEditorTestNodes" Value="$(GeneratedEditorTestNodes);$(PlatformEditorTestNodes)" />
@@ -259,19 +259,19 @@
 		<Property Name="TestAgentType" Value="$(TestAgentOverride)" If="'$(TestAgentOverride)' != ''"/>
 
 		<!-- Building and cooking will require these nodes from the appropriate host. They should have been built above... -->
-		<Property Name="PlatformToolsNodeName" Value="Compile $(ProjectName)Tools $(HostAgentType)" />
+		<Property Name="PlatformToolsNodeName" Value="Compile Tools $(HostAgentType)" />
 		<Property Name="PlatformEditorNodeName" Value="Compile $(ProjectName)Editor $(HostAgentType)" />
 		<!-- Define these names up front as we use some of them several times -->
 		<Property Name="PlatformAgentName" Value="$(DisplayTargetName) $(TargetPlatform)" />
 		<Property Name="PlatformCompileNodeName" Value="Compile $(DisplayTargetName) $(TargetPlatform)" />
 		<Property Name="UploadSymbolsNodeName" Value="" />
-		<Property Name="UploadSymbolsNodeName" Value="Upload Symbols $(TargetName) $(TargetPlatform)" If="'$(SymbolServerPath)' != ''" />
-		<Property Name="PlatformCookNodeName" Value="Cook $(DisplayTargetName) $(TargetPlatform)" />
-		<Property Name="PlatformStageNodeName" Value="Stage $(DisplayTargetName) $(TargetPlatform)" />
-		<Property Name="PlatformPackageNodeName" Value="Package $(DisplayTargetName) $(TargetPlatform)" />
-		<Property Name="PlatformHashNodeName" Value="Hash $(DisplayTargetName) $(TargetPlatform)" />
-		<Property Name="PlatformPublishStagedNodeName" Value="Publish Staged $(DisplayTargetName) $(TargetPlatform)" />
-		<Property Name="PlatformPublishPackagedNodeName" Value="Publish Packaged $(DisplayTargetName) $(TargetPlatform)" />
+		<Property Name="UploadSymbolsNodeName" Value="Upload Symbols $(TargetPlatform)" If="'$(SymbolServerPath)' != ''" />
+		<Property Name="PlatformCookNodeName" Value="Cook $(TargetPlatform)" />
+		<Property Name="PlatformStageNodeName" Value="Stage $(TargetPlatform)" />
+		<Property Name="PlatformPackageNodeName" Value="Package $(TargetPlatform)" />
+		<Property Name="PlatformHashNodeName" Value="Hash $(TargetPlatform)" />
+		<Property Name="PlatformPublishStagedNodeName" Value="Publish Staged $(TargetPlatform)" />
+		<Property Name="PlatformPublishPackagedNodeName" Value="Publish Packaged $(TargetPlatform)" />
 		
 		<!-- Set CookPlatform since it isn't always the same as the target -->
 		<Property Name="CookPlatform" Value="$(TargetPlatform)" />
@@ -388,7 +388,7 @@
 				<Property Name="GauntletTestArgs" Value="$(GauntletTargetTestAltArgs)" If="ContainsItem('$(AlternateDevicePoolPlatforms)', '$(TargetPlatform)', ';')"/>
 
 				<!-- Declare a boot test for the target-->	
-				<Property Name="BootTestName" Value="$(DisplayTargetName) $(TargetPlatform) $(TestConfiguration) BootTest" />
+				<Property Name="BootTestName" Value="$(TargetPlatform) $(TestConfiguration) BootTest" />
 				<Node Name="$(BootTestName)" After="$(PlatformPublishPackagedNodeName)" Requires="$(TestPrerequisites)">
 					<Do If="'$(AutomationReportOutputDirectory)' != ''">
 						<Property Name="OutputFragmentPath" Value="$(TargetPlatform)/$(TestConfiguration)/UE.BootTest" />
@@ -414,7 +414,7 @@
 						<Property Name="GauntletReportArgs" Value="-logdir=&quot;$(OutputReportPath)&quot; -artifactname=Gauntlet -ReportExportPath=&quot;$(OutputReportPath)&quot;" />
 						<Property Name="GauntletReportArgs" Value="$(GauntletReportArgs) -ReportURL=&quot;$(AutomationReportUrl)/$(OutputFragmentPath)/&quot;" If="'$(AutomationReportUrl)' != ''" />
 					</Do>
-					<Property Name="TestNodeName" Value="$(DisplayTargetName) $(TargetPlatform) $(TestConfiguration) $(TestName)" />
+					<Property Name="TestNodeName" Value="$(TargetPlatform) $(TestConfiguration) $(TestName)" />
 					<!-- Run the test. No 'Requires' as we run on the same machine locally and use the publish folder otherwise -->
 					<Node Name="$(TestNodeName)" Requires="$(TestPrerequisites)">
 						<Command Name="RunUnreal" Arguments="-project=$(TargetProject) -test=&quot;$(TestName)&quot; -platform=$(TargetPlatform) -configuration=$(TestConfiguration) $(GauntletBuildArg) $(GauntletTestArgs) $(GauntletReportArgs)"/>
@@ -428,12 +428,11 @@
 		</Agent>
 
 		<!-- Declare labels for CIS -->
-		<Label Category="Platforms" Name="$(ProjectName) $(TargetPlatform)" Requires="$(PlatformCompileNodeName);$(UploadSymbolsNodeName)" />
-		<Label Category="Cook" Name="Cook $(DisplayTargetName) $(TargetPlatform)" Requires="$(PlatformCookNodeName)" />
-		<Label Category="Stage" Name="Stage $(DisplayTargetName) $(TargetPlatform)" Requires="$(PlatformStageNodeName);$(PlatformHashNodeName);$(PlatformPublishStagedNodeName)" Exclude="$(PlatformCookNodeName)" />
-		<Label Category="Package" Name="Pkg $(DisplayTargetName) $(TargetPlatform)" Requires="$(PlatformPackageNodeName);$(PlatformPublishPackagedNodeName)" Exclude="$(PlatformPublishStagedNodeName)" />
-		<Label Category="Test" Name="BootTest $(DisplayTargetName) $(TargetPlatform)" Requires="$(AllPlatformBootTestNodes)" Exclude="$(PlatformEditorNodeName);$(PlatformPackageNodeName);$(PlatformPublishPackagedNodeName)" />
-		<Label Category="Test" Name="Test $(DisplayTargetName) $(TargetPlatform)" Requires="$(AllPlatformTestNodes)" Exclude="$(PlatformEditorNodeName);$(PlatformPackageNodeName);$(AllPlatformBootTestNodes)"/>
+		<Label Category="Clients" Name="$(ProjectName) $(TargetPlatform)" Requires="$(PlatformCompileNodeName);$(UploadSymbolsNodeName)" />
+		<Label Category="Clients" Name="Content $(DisplayTargetName) $(TargetPlatform)" Requires="$(PlatformCookNodeName);$(PlatformStageNodeName);$(PlatformPackageNodeName)" />
+		<Label Category="Clients" Name="Publish $(DisplayTargetName) $(TargetPlatform)" Requires="$(PlatformPublishPackagedNodeName);$(PlatformPublishPackagedNodeName);$(PlatformPublishStagedNodeName);$(PlatformHashNodeName)" />
+		<Label Category="Tests" Name="BootTest $(DisplayTargetName) $(TargetPlatform)" Requires="$(AllPlatformBootTestNodes)" />
+		<Label Category="Tests" Name="Test $(DisplayTargetName) $(TargetPlatform)" Requires="$(AllPlatformTestNodes)" />
 
 		<!-- Tie all steps into our platform tasks unless excluded -->
 		<Property Name="PlatformTasks" Value="" />
