Remove Trigger code from BuildGraph example. We don't support this element any more.

diff --git a/Engine/Build/Graph/Examples/TagsAndFiles.xml b/Engine/Build/Graph/Examples/TagsAndFiles.xml
--- a/Engine/Build/Graph/Examples/TagsAndFiles.xml
+++ b/Engine/Build/Graph/Examples/TagsAndFiles.xml
@@ -103,19 +103,17 @@
 		</Node>
 	</Agent>
 
-	<Trigger Name="Test Trigger">
-		<Agent Name="Trigger Group" Type="CompileWin64">
-			<Node Name="Print all Output Folders" Requires="Manipulate Some Files;Create A Named Output;Copy Syntax">
-				<!-- Test that we have all the side-effects of the upstream node -->
-				<Log Message="Here are all the files in Engine/TestNumber1/..." Files="Engine/TestNumber1/..."/>
-				<Log Message="Here are all the files in Engine/TestNumber2/..." Files="Engine/TestNumber2/..."/>
-				<Log Message="Here are all the files in Engine/TestNumber3/..." Files="Engine/TestNumber3/..."/>
-				<Log Message="Here are all the files in Engine/TestNumber4/..." Files="Engine/TestNumber4/..."/>
-				<Log Message="Here are all the files in Engine/TestNumber5/..." Files="Engine/TestNumber5/..."/>
-				<Log Message="Here are all the files in Engine/TestNumber6/..." Files="Engine/TestNumber6/..."/>
-			</Node>
-		</Agent>
-	</Trigger>
+	<Agent Name="Trigger Group" Type="CompileWin64">
+		<Node Name="Print all Output Folders" Requires="Manipulate Some Files;Create A Named Output;Copy Syntax">
+			<!-- Test that we have all the side-effects of the upstream node -->
+			<Log Message="Here are all the files in Engine/TestNumber1/..." Files="Engine/TestNumber1/..."/>
+			<Log Message="Here are all the files in Engine/TestNumber2/..." Files="Engine/TestNumber2/..."/>
+			<Log Message="Here are all the files in Engine/TestNumber3/..." Files="Engine/TestNumber3/..."/>
+			<Log Message="Here are all the files in Engine/TestNumber4/..." Files="Engine/TestNumber4/..."/>
+			<Log Message="Here are all the files in Engine/TestNumber5/..." Files="Engine/TestNumber5/..."/>
+			<Log Message="Here are all the files in Engine/TestNumber6/..." Files="Engine/TestNumber6/..."/>
+		</Node>
+	</Agent>
 
 	<!-- Aggregates for the build process -->
 	<Aggregate Name="Tag and File Examples" Requires="Print all Output Folders;Consume a Single Named Output;List Named Outputs"/>
