Add support for CookedEditor to BuildAndTestProject.xml

diff --git a/Engine/Build/Graph/Tasks/BuildAndTestProject.xml b/Engine/Build/Graph/Tasks/BuildAndTestProject.xml
--- a/Engine/Build/Graph/Tasks/BuildAndTestProject.xml
+++ b/Engine/Build/Graph/Tasks/BuildAndTestProject.xml
@@ -10,6 +10,7 @@
 
 	<Option Name="TargetName" DefaultValue="" Restrict="" Description="If your project supports both Game and Client you can set this, otherwise one will be selected"/>
 	<Option Name="EditorTargetName" DefaultValue="" Restrict="" Description="Specify the editor target name if different than (ProjectName)Editor"/>
+	<Option Name="CookedEditorTargetName" DefaultValue="" Restrict="" Description="Specify the cooked editor target name if different than (ProjectName)CookedEditor"/>
 
 	<Option Name="GenericCompileArguments" DefaultValue="" Restrict="" Description="Configurable compile args applied to all compile targets."/>
 	<Option Name="ExtraToolCompileArguments" DefaultValue="" Restrict="" Description="Extra arguments to use when building the tools" />
@@ -27,26 +28,21 @@
 	<Option Name="TestAgentOverride" DefaultValue="" Description="If set, other logic is ignored and this agent type is used for all test work." />
 	<Option Name="ClientCompileAgentOverride" DefaultValue="" Description="Specifies the agent type to use for C++ compilation jobs for target platform binaries." />
 
-	<Option Name="WithBATDefaults" DefaultValue="true" Description="Used by parent scripts to specify defaults when including this script"/>
-
 	<!-- 
-		Allow project graphs that Include this file to specify their own defaults. This is done by specifying WithBATDefaults=false
-		then declaring *all* these default values (leaving them empty if no default is required). 
-		The reason we need to wrap these in WithBATDefaults is that the properties always neeed to be declared.
-		Note - this means if a parent script includes this script multiple tiles then it should
-		always WithBATDefaults true/false before doing so, and if false it should supply all these arguments.
+		Allow project graphs that Include this file to specify their own defaults.
 	-->
-	<Do If="$(WithBATDefaults)">
-		<Property Name="DefaultEditorPlatforms" Value="" />
-		<Property Name="DefaultTargetPlatforms" Value="" />
-		<Property Name="DefaultTargetConfigurations" Value="Development" />
-		<Property Name="DefaultEditorTestList" Value="" />
-		<Property Name="DefaultTargetTestList" Value="" />
-		<Property Name="NetworkTempRootOverride" Value="" />
-		<Property Name="NetworkPublishRootOverride" Value="" />
-		<Property Name="NetworkReportRootOverride" Value="" />
-		<Property Name="SymbolServerPath" Value="" />
-	</Do>
+	<Option Name="DefaultEditorPlatforms" DefaultValue="" Description="Default Editor platforms"/>
+	<Option Name="DefaultCookedEditorPlatforms" DefaultValue="" Description="Default Cooked Editor platforms"/>
+	<Option Name="DefaultTargetPlatforms" DefaultValue="" Description="Default Target platforms"/>
+	<Option Name="DefaultTargetConfigurations" DefaultValue="Development" Description="Default Target configurations"/>
+	<Option Name="DefaultCookedEditorConfigurations" DefaultValue="Development" Description="Default Cooked Editor configurations"/>
+	<Option Name="DefaultEditorTestList" DefaultValue="" Description="Default Editor test list"/>
+	<Option Name="DefaultTargetTestList" DefaultValue="" Description="Default Target test list"/>
+	<Option Name="DefaultCookedEditorTestList" DefaultValue="" Description="Default Cooked Editor test list"/>
+	<Option Name="NetworkTempRootOverride" DefaultValue="" Description="Network temp root override"/>
+	<Option Name="NetworkPublishRootOverride" DefaultValue="" Description="Network publish root override"/>
+	<Option Name="NetworkReportRootOverride" DefaultValue="" Description="Network report root override"/>
+	<Option Name="SymbolServerPath" DefaultValue="" Description="Symbol server path"/>
 
 	<Option Name="EnableSymbolStore" DefaultValue="$(IsBuildMachine)" Restrict="true|false" Description="Enables the upload of debugging symbols" />
 	<Property Name="SymbolServerPath" Value="" If="$(EnableSymbolStore) == 'false'"/>
@@ -58,8 +54,10 @@
 	<Do If="'$(PlatformsToSkipTests)' != ''" >
 		<Property Name="EditorPlatformsToSkipTests" Value="$(PlatformsToSkipTests)" />
 		<Property Name="TargetPlatformsToSkipTests" Value="$(PlatformsToSkipTests)" />
+		<Property Name="CookedEditorPlatformsToSkipTests" Value="$(PlatformsToSkipTests)" />
 		<Property Name="EditorPlatformsToSkipBootTests" Value="$(PlatformsToSkipTests)" />
 		<Property Name="TargetPlatformsToSkipBootTests" Value="$(PlatformsToSkipTests)" />
+		<Property Name="CookedEditorPlatformsToSkipBootTests" Value="$(PlatformsToSkipTests)" />
 	</Do>
 
 	<!-- Does this project require a dedicated server for testing? (not currently implemented!)-->
@@ -76,13 +74,16 @@
 	<!-- Defines whether we always run a simple 'does it boot' test before running other tests. This can save a lot of time/repeated errors -->
 	<Option Name="EditorBootTest" DefaultValue="!$(SkipTest)" Description="Run a simple boot test on the editor before other tests"/>
 	<Option Name="TargetBootTest" DefaultValue="!$(SkipTest)" Description="Run a simple boot test on target platforms before other tests"/>
+	<Option Name="CookedEditorBootTest" DefaultValue="!$(SkipTest)" Description="Run a simple boot test on cooked editor platforms before other tests"/>
 
 	<Option Name="FillEditorDDC" DefaultValue="false" Description="Whether to fill the editor DDC before running any editor tests"/>
 
 	<Option Name="EditorTestList" DefaultValue="$(DefaultEditorTestList)" Description="Tests to run on the editor. Test1+Test2 etc"/>
 	<Option Name="TargetTestList" DefaultValue="$(DefaultTargetTestList)" Description="Tests to run on the target platform(s). Test1+Test2 etc"/>
+	<Option Name="CookedEditorTestList" DefaultValue="$(DefaultCookedEditorTestList)" Description="Tests to run on the cooked editor. Test1+Test2 etc"/>
 	<Option Name="EditorPlatformsToSkipBootTests" DefaultValue="$(EditorPlatformsToSkipTests)" Description="Editor platforms that should skip boot tests."/>
 	<Option Name="TargetPlatformsToSkipBootTests" DefaultValue="$(TargetPlatformsToSkipTests)" Description="Target platforms that should skip boot tests."/>
+	<Option Name="CookedEditorPlatformsToSkipBootTests" DefaultValue="$(CookedEditorPlatformsToSkipTests)" Description="Cooked editor platforms that should skip boot tests."/>
 
 	<!-- Settings files that can be overriden -->
 	<Option Name="PathSettingsFile" DefaultValue="Inc/PathSettings.xml" Description="Path settings to use" />
@@ -107,13 +108,16 @@
 	<Property Name="TargetName" Value="$(ProjectName)" If="'$(TargetName)' == ''"/>
 	<Property Name="EditorTargetName" Value="UnrealEditor" If="$(IsContentOnlyProject) And '$(EditorTargetName)' == ''"/>
 	<Property Name="EditorTargetName" Value="$(ProjectName)Editor" If="'$(EditorTargetName)' == ''"/>
+	<Property Name="CookedEditorTargetName" Value="$(ProjectName)CookedEditor" If="'$(CookedEditorTargetName)' == ''"/>
 
 	<!-- If skip test is true, add all target platforms to the "do not test" list -->
 	<Do If="$(SkipTest)">
 		<Property Name="EditorPlatformsToSkipTests" Value="$(EditorPlatforms)" />
 		<Property Name="TargetPlatformsToSkipTests" Value="$(TargetPlatforms)" />
+		<Property Name="CookedEditorPlatformsToSkipTests" Value="$(CookedEditorPlatforms)" />
 		<Property Name="EditorPlatformsToSkipBootTests" Value="$(EditorPlatformsToSkipTests)" />
 		<Property Name="TargetPlatformsToSkipBootTests" Value="$(TargetPlatformsToSkipTests)" />
+		<Property Name="CookedEditorPlatformsToSkipBootTests" Value="$(CookedEditorPlatformsToSkipTests)" />
 	</Do>
 
 	<!-- If skip packaging is true, add all target platforms to the "do not test" list -->
@@ -140,6 +144,8 @@
 	<Property Name="GeneratedEditorTestNodes" Value="" />
 	<Property Name="GeneratedTargetNodes" Value="" />
 	<Property Name="GeneratedTargetTestNodes" Value="" />
+	<Property Name="GeneratedCookedEditorNodes" Value="" />
+	<Property Name="GeneratedCookedEditorTestNodes" Value="" />
 
 	<!-- Do all editor steps under a single agent as the things that can be done in parallel probably aren't worth the overhead-->
 	<ForEach Name="EditorPlatform" Values="$(RequiredEditorPlatforms)" Separator="+">
@@ -179,8 +185,16 @@
 			<Node Name="Compile $(ProjectName)Editor $(EditorPlatform)" Requires="Compile Tools $(EditorPlatform)">
 				<Compile Target="$(EditorTargetName)" Platform="$(EditorPlatform)" Configuration="Development" Arguments="$(GenericCompileArguments) $(ExtraEditorCompileArguments)"/>
 			</Node>
+			<!-- Cooked editor target -->
+			<Node Name="Compile $(ProjectName)CookedEditor $(EditorPlatform)" Requires="Compile Tools $(EditorPlatform)">
+				<ForEach Name="Config" Values="$(CookedEditorConfigurations)" Separator="+">
+					<Compile Target="$(CookedEditorTargetName)" Platform="$(EditorPlatform)" Configuration="$(Config)" Arguments="$(GenericCompileArguments) $(ExtraEditorCompileArguments)"/>
+				</ForEach>
+			</Node>
 			<!--Add these nodes to our dependency list -->
 			<Property Name="GeneratedToolNodes" Value="Compile $(ProjectName)Editor $(EditorPlatform); Compile Tools $(EditorPlatform)" If="!$(SkipBuild)" />
+			<!-- Editor Labels -->
+			<Label Category="Editors" Name="$(ProjectName) Editor $(EditorPlatform)" Requires="$(GeneratedToolNodes)" />
 		</Agent>
 		<Agent Name="Test $(ProjectName)Editor $(EditorPlatform)" Type="$(TestAgentType)">
 			<!-- Prerequisites for the editor test to run -->
@@ -223,7 +237,6 @@
 			</ForEach>
 
 			<!-- Create Labels -->
-			<Label Category="Editors" Name="$(ProjectName) Editor $(EditorPlatform)" Requires="$(GeneratedToolNodes)" />
 			<Label Category="Editors" Name="$(ProjectName) Editor DDC $(EditorPlatform)" Requires="$(ProjectEditorDDC)" />
 
 			<!-- If this editor was requested (as opposed to being required due to target platforms) then link in the test nodes and set up labels -->
@@ -237,6 +250,98 @@
 		</Agent>
 	</ForEach>
 
+	<!-- Cooked editor Targets -->
+	<ForEach Name="EditorPlatform" Values="$(CookedEditorPlatforms)" Separator="+">
+		<Property Name="HostAgentType" Value="$(EditorPlatform)" />
+		<Property Name="TestAgentType" Value="Test$(EditorPlatform)" />
+		<Property Name="TestAgentType" Value="$(TestAgentOverride)" If="'$(TestAgentOverride)' != ''"/>
+
+		<Property Name="PlatformPublishNodeName" Value="Publish CookedEditor $(EditorPlatform)" />
+
+		<Property Name="CookedEditorCookRequirements" Value="Compile Tools $(EditorPlatform);Compile $(ProjectName)Editor $(EditorPlatform);Compile $(ProjectName)CookedEditor $(EditorPlatform)"/>
+		<Property Name="CookedEditorCookerArgs" Value="-project=$(TargetProject) -platform=$(EditorPlatform) -compressed $(ExtraCookArguments)"/>
+
+		<Agent Name="CookedEditor Cook and Publish $(EditorPlatform)" Type="$(HostAgentType)">
+			<!-- Cook Editor -->
+			<Node Name="Cook CookedEditor $(EditorPlatform)" Requires="$(CookedEditorCookRequirements)">
+				<Command Name="MakeCookedEditor" Arguments="-cook $(CookedEditorCookerArgs)"/>
+			</Node>
+
+			<!-- Stage CookedEditor -->
+			<Property Name="CookedEditorStageArgs" Value="-skipcook -stage -pak -platform=$(EditorPlatform) -clientconfig=$(CookedEditorConfigurations) -StagingDirectory=&quot;$(ProjectOutputDirectory)&quot;" />
+			<Property Name="CookedEditorStageRequirements" Value="$(CookedEditorCookRequirements)"/>
+			<Property Name="CookedEditorStageRequirements" Value="$(CookedEditorStageRequirements);Cook CookedEditor $(EditorPlatform)" If="!$(SkipCook)"/>
+
+			<Node Name="Stage CookedEditor $(EditorPlatform)" Requires="$(CookedEditorStageRequirements)">
+				<Command Name="MakeCookedEditor" Arguments="-project=$(TargetProject) $(CookedEditorStageArgs)"/>
+			</Node>
+
+			<!-- Publish CookedEditor -->
+			<Node Name="Publish CookedEditor $(EditorPlatform)" Requires="Stage CookedEditor $(EditorPlatform)">
+				<Log Message="Copying to $(NetworkOutputDirectory)/$(EditorPlatform)CookedEditor"/>
+				<Do If="!$(SkipPublish)">
+					<Copy Files="..." From="$(ProjectOutputDirectory)/$(EditorPlatform)CookedEditor" To="$(NetworkOutputDirectory)/$(EditorPlatform)CookedEditor"/>
+				</Do>
+			</Node>
+
+			<!-- Aggregate Nodes requirements -->
+			<Property Name="GeneratedCookedEditorNodes" Value="Compile $(ProjectName)CookedEditor $(EditorPlatform)" If="!$(SkipBuild)" />
+			<Property Name="GeneratedCookedEditorNodes" Value="$(GeneratedCookedEditorNodes);Cook CookedEditor $(EditorPlatform)" If="!$(SkipCook)" />
+			<Property Name="GeneratedCookedEditorNodes" Value="$(GeneratedCookedEditorNodes);$(PlatformPublishNodeName)" If="!$(SkipPublish) And $(PublishBuilds)" />
+		</Agent>
+
+		<!-- Tests -->
+		<Agent Name="Test $(ProjectName)CookedEditor $(EditorPlatform)" Type="$(TestAgentType)">
+			<!-- Prerequisites for the editor test to run -->
+			<Property Name="TestPrerequisites" Value="$(GeneratedCookedEditorNodes)"/>
+			<Property Name="TestPrerequisites" Value="$(TestPrerequisites);$(PlatformPublishNodeName)" If="!$(PublishBuilds)"/>
+			<!-- Declare a boot test for the editor-->
+			<Property Name="PlatformBootTestNodeName" Value="CookedEditor $(EditorPlatform) Test=BootTest" />
+			<Do If="'$(AutomationReportOutputDirectory)' != ''">
+				<Property Name="OutputFragmentPath" Value="$(EditorPlatform)Editor/UE.CookedEditorBootTest" />
+				<Property Name="OutputReportPath" Value="$(AutomationReportOutputDirectory)/$(OutputFragmentPath)" />
+				<Property Name="GauntletReportArgs" Value="-logdir=&quot;$(OutputReportPath)&quot; -artifactname=Gauntlet -ReportExportPath=&quot;$(OutputReportPath)&quot;" />
+				<Property Name="GauntletReportArgs" Value="$(GauntletReportArgs) -ReportURL=&quot;$(AutomationReportUrl)/$(OutputFragmentPath)/&quot;" If="'$(AutomationReportUrl)' != ''" />
+			</Do>
+			<Node Name="$(PlatformBootTestNodeName)" Requires="$(TestPrerequisites)">
+				<Command Name="RunUnreal" Arguments="-project=$(TargetProject) -build=&quot;$(NetworkOutputDirectory)/$(EditorPlatform)CookedEditor&quot; -cookededitor -test=&quot;UE.EditorBootTest&quot; $(GauntletEditorTestArgs) $(GauntletReportArgs)" />
+			</Node>
+
+			<!-- Link the boot test in to the prerequisites and generated nodes if its enabled-->
+			<Do If="$(CookedEditorBootTest)">
+				<Property Name="TestPrerequisites" Value="$(TestPrerequisites);$(PlatformBootTestNodeName)" />
+			</Do>
+			<!-- Now generate a test node for each editor test in the list-->
+			<Property Name="PlatformCookedEditorTestNodes" Value="" />
+			<ForEach Name="TestName" Values="$(CookedEditorTestList)" Separator="+">
+				<Do If="'$(AutomationReportOutputDirectory)' != ''">
+					<Property Name="OutputFragmentPath" Value="$(EditorPlatform)/Editor$(TestName)" />
+					<Property Name="OutputReportPath" Value="$(AutomationReportOutputDirectory)/$(OutputFragmentPath)" />
+					<Property Name="GauntletReportArgs" Value="-logdir=&quot;$(OutputReportPath)&quot; -artifactname=Gauntlet -ReportExportPath=&quot;$(OutputReportPath)&quot;" />
+					<Property Name="GauntletReportArgs" Value="$(GauntletReportArgs) -ReportURL=&quot;$(AutomationReportUrl)/$(OutputFragmentPath)/&quot;" If="'$(AutomationReportUrl)' != ''" />
+				</Do>
+				<Property Name="TestNodeName" Value="CookedEditor $(EditorPlatform) Test=$(TestName)" />
+				<Node Name="$(TestNodeName)" Requires="$(TestPrerequisites)">
+					<Command Name="RunUnreal" Arguments="-project=$(TargetProject) -build=&quot;$(NetworkOutputDirectory)/$(EditorPlatform)CookedEditor&quot; -cookededitor -test=&quot;$(TestName)&quot; $(GauntletReportArgs) $(GauntletEditorTestArgs)" />
+				</Node>
+				<!-- Add each test node to our generated list -->
+				<Property Name="PlatformCookedEditorTestNodes" Value="$(PlatformCookedEditorTestNodes);$(TestNodeName)" />
+			</ForEach>
+
+			<!-- Create Labels -->
+			<!-- If this cooked editor was requested (as opposed to being required due to target platforms) then link in the test nodes and set up labels -->
+			<Do If="ContainsItem('$(CookedEditorPlatforms)', '$(EditorPlatform)', '+') and !ContainsItem('$(CookedEditorPlatformsToSkipTests)', '$(EditorPlatform)', '+')">
+				<Label Category="Tests" Name="BootTest $(ProjectName)CookedEditor $(EditorPlatform)" Requires="$(PlatformBootTestNodeName)" />
+				<Label Category="Tests" Name="Test $(ProjectName)CookedEditor $(EditorPlatform)" Requires="$(PlatformCookedEditorTestNodes)" />
+				<Property Name="GeneratedCookedEditorTestNodes" Value="$(GeneratedCookedEditorTestNodes);$(PlatformBootTestNodeName)" If="$(CookedEditorBootTest)" />
+				<Property Name="GeneratedCookedEditorTestNodes" Value="$(GeneratedCookedEditorTestNodes);$(PlatformCookedEditorTestNodes)" />
+			</Do>
+		</Agent>
+
+		<!-- Add to Editor Labels -->
+		<Label Category="Editors" Name="$(ProjectName) CookedEditor $(EditorPlatform)" Requires="$(GeneratedCookedEditorNodes)" />
+	</ForEach>
+
 	<!-- Now, for each platform we build, cook, stage, package, publish, and test -->
 	<ForEach Name="TargetPlatform" Values="$(TargetPlatforms)" Separator="+">
 
@@ -446,6 +551,6 @@
 	</ForEach>
 
 	<!-- Create an aggregate with all the generated nodes -->
-	<Aggregate Name="BuildAndTest $(ProjectName)" Requires="$(GeneratedToolNodes);$(GeneratedEditorTestNodes);$(GeneratedTargetNodes);$(GeneratedTargetTestNodes)"/>
+	<Aggregate Name="BuildAndTest $(ProjectName)" Requires="$(GeneratedToolNodes);$(GeneratedCookedEditorNodes);$(GeneratedEditorTestNodes);$(GeneratedCookedEditorTestNodes);$(GeneratedTargetNodes);$(GeneratedTargetTestNodes)"/>
 
 </BuildGraph>
\ No newline at end of file
diff --git a/Engine/Build/Graph/Tasks/Inc/PlatformOptions.xml b/Engine/Build/Graph/Tasks/Inc/PlatformOptions.xml
--- a/Engine/Build/Graph/Tasks/Inc/PlatformOptions.xml
+++ b/Engine/Build/Graph/Tasks/Inc/PlatformOptions.xml
@@ -26,20 +26,23 @@
 
 <Property Name="ShowExplicitFlags" Value="true" />
 
-<Property Name="AllEditorPlatforms" Value="Mac+Win64" />
+<Property Name="AllEditorPlatforms" Value="Mac+Win64+Linux" />
 <Property Name="AllTargetPlatforms" Value="Android+HoloLens+iOS+Mac+PS4+PS5+Stadia+Switch+tvOS+Win64+WinGDK+XboxOneGDK+XB1+XSX" />
 
 <Option Name="EditorPlatforms" DefaultValue="" Description="Alternative to 'WithFooEditor', a list of editors to build and test. Win64+Mac etc"/>
 <Option Name="TargetPlatforms" DefaultValue="" Description="Alternative to 'WithFoo', a list of Platforms to build and test. PS4+XSX etc"/>
+<Option Name="CookedEditorPlatforms" DefaultValue="" Description="Alternative to 'WithFooCookedEditor', a list of editors to build and cooked. Win64+Mac etc"/>
 
 <!-- Support simple One+Two+Three style -->
 <Option Name="TargetConfigurations" DefaultValue="" Description="Configuration(s) to build"/>
 <Option Name="PackageConfigurations" DefaultValue="$(TargetConfigurations)" Description="Configuration to package"/>
 <Option Name="TestConfigurations" DefaultValue="$(PackageConfigurations)" Description="Configuration to test"/>
+<Option Name="CookedEditorConfigurations" DefaultValue="" Description="Configuration(s) to use when making a cooked editor (Development, Shipping or Development+Shipping)"/>
 
 <!-- No means of reliably running these tests atm -->
 <Option Name="EditorPlatformsToSkipTests" DefaultValue="" Restrict=".*" Description="List of editor platforms that editor tests won't be run on" />
 <Option Name="TargetPlatformsToSkipTests" DefaultValue="HoloLens+IOS+Linux+Stadia+TVOS" Restrict=".*" Description="List of target platforms that target tests won't be run on" />
+<Option Name="CookedEditorPlatformsToSkipTests" DefaultValue="" Restrict=".*" Description="List of editor cooked platforms that editor tests won't be run on" />
 <Option Name="PlatformsToSkipTests" DefaultValue="" Restrict=".*" Description="List of platforms that all tests won't be run on. Overrides Editor/TargetPlatformsToSkipTests when non-empty." />
 
 <!-- Support explicit flags for horde but hide them otherwise as the option list is overwhelming..-->
@@ -48,6 +51,8 @@
 	<Option Name="WithMacEditor" DefaultValue="false" Description="Build and test Mac Editor"/>
 	<Option Name="WithLinuxEditor" DefaultValue="false" Description="Build and test Linux Editor"/>
 
+	<Option Name="WithCookedEditor" DefaultValue="false" Description="Build and cooked CookedEditor"/>
+
 	<Option Name="WithAndroid" DefaultValue="false" Description="Build and test Android"/>
 	<Option Name="WithHoloLens" DefaultValue="false" Description="Build and test HoloLens"/>
 	<Option Name="WithIOS" DefaultValue="false" Description="Build and test IOS"/>
@@ -71,6 +76,13 @@
 		<Property Name="EditorPlatforms" Value="$(EditorPlatforms)Linux+" If="$(WithLinuxEditor)" />
 	</Do>
 
+	<Do If="'$(CookedEditorPlatforms)' == ''">
+		<Property Name="CookedEditorPlatforms" Value="" />
+		<Property Name="CookedEditorPlatforms" Value="$(CookedEditorPlatforms)Win64+" If="$(WithWin64Editor) And $(WithCookedEditor)" />
+		<Property Name="CookedEditorPlatforms" Value="$(CookedEditorPlatforms)Mac+" If="$(WithMacEditor) And $(WithCookedEditor)" />
+		<Property Name="CookedEditorPlatforms" Value="$(CookedEditorPlatforms)Linux+" If="$(WithLinuxEditor) And $(WithCookedEditor)" />
+	</Do>
+
 	<Do If="'$(TargetPlatforms)' == ''" >
 		<Property Name="TargetPlatforms" Value="" />
 		<Property Name="TargetPlatforms" Value="$(TargetPlatforms)Android+" If="$(WithAndroid)" />
@@ -96,6 +108,10 @@
 	<Property Name="EditorPlatforms" Value="$(DefaultEditorPlatforms)" />
 </Do>
 
+<Do If="'$(CookedEditorPlatforms)' == ''" >
+	<Property Name="CookedEditorPlatforms" Value="$(DefaultCookedEditorPlatforms)" />
+</Do>
+
 <Do If="'$(TargetPlatforms)' == ''" >
 	<Property Name="TargetPlatforms" Value="$(DefaultTargetPlatforms)" />
 </Do>
@@ -104,6 +120,9 @@
 <Do If="'$(EditorPlatforms)' == 'None'" >
 	<Property Name="EditorPlatforms" Value="" />
 </Do>
+<Do If="'$(CookedEditorPlatforms)' == 'None'" >
+	<Property Name="CookedEditorPlatforms" Value="" />
+</Do>
 <Do If="'$(TargetPlatforms)' == 'None'" >
 	<Property Name="TargetPlatforms" Value="" />
 </Do>
@@ -115,6 +134,9 @@
 <Do If="'$(TargetPlatforms)' == 'All'" >
 	<Property Name="TargetPlatforms" Value="$(AllTargetPlatforms)" />
 </Do>
+<Do If="'$(CookedEditorPlatforms)' == 'All'" >
+	<Property Name="CookedEditorPlatforms" Value="$(AllEditorPlatforms)"/>
+</Do>
 
 <Property Name="RequiredEditorPlatforms" Value="$(EditorPlatforms)" />
 
@@ -132,6 +154,16 @@
 	</Do>
 </ForEach>
 
+<ForEach Name="CookedEditorPlatform" Values="$(CookedEditorPlatforms)" Separator="+">
+	<!-- Check the platform is valid -->
+	<Do If="!ContainsItem('$(AllEditorPlatforms)','$(CookedEditorPlatform)','+')" >
+		<Error Message="Platform $(CookedEditorPlatform) is not a recognized platform!" />
+	</Do>
+	<Do If="!ContainsItem('$(RequiredEditorPlatforms)','$(CookedEditorPlatform)', '+')">
+		<Property Name="RequiredEditorPlatforms" Value="$(RequiredEditorPlatforms)+$(CookedEditorPlatform)" />
+	</Do>
+</ForEach>
+
 <!-- PS5 packaging isn't supported by Gauntlet yet :( )-->
 <Property Name="PlatformsToSkipPackaging" Value="Win64+Linux" />
 
@@ -155,6 +187,9 @@
 	<Option Name="WithTestTests" DefaultValue="false" Description="Run tests on Test builds"/>
 	<Option Name="WithShippingTests" DefaultValue="false" Description="Run tests on Shipping builds"/>
 
+	<Option Name="WithDevelopmentCookedEditor" DefaultValue="false" Description="Build Development configuration for Editor cooked"/>
+	<Option Name="WithShippingCookedEditor" DefaultValue="false" Description="Build Shipping configuration for Editor cooked"/>
+
 	<Do If="'$(TargetConfigurations)' == ''" >	
 		<Property Name="TargetConfigurations" Value="" />
 		<Property Name="TargetConfigurations" Value="$(TargetConfigurations)Debug+"  If="$(WithDebugTargets)" />
@@ -178,6 +213,12 @@
 		<Property Name="TestConfigurations" Value="$(TestConfigurations)Test+"  If="$(WithTestTests)" />
 		<Property Name="TestConfigurations" Value="$(TestConfigurations)Shipping+"  If="$(WithShippingTests)" />
 	</Do>
+
+	<Do If="'$(CookedEditorConfigurations)' == ''" >	
+		<Property Name="CookedEditorConfigurations" Value="" />
+		<Property Name="CookedEditorConfigurations" Value="$(CookedEditorConfigurations)Development+"  If="$(WithDevelopmentCookedEditor)" />
+		<Property Name="CookedEditorConfigurations" Value="$(CookedEditorConfigurations)Shipping+"  If="$(WithShippingCookedEditor)" />
+	</Do>
 </Do>
 
 <Do If="'$(TargetConfigurations)' == ''" >
@@ -192,6 +233,10 @@
 	<Property Name="TestConfigurations" Value="$(PackageConfigurations)" />
 </Do>
 
+<Do If="'$(CookedEditorConfigurations)' == ''" >	
+	<Property Name="CookedEditorConfigurations" Value="$(DefaultCookedEditorConfigurations)" />
+</Do>
+
 <!-- add up any package configurations that are needed by test configurations.. -->
 <ForEach Name="TestConfiguration" Values="$(TestConfigurations)" Separator="+">
 	<Do If="!ContainsItem('$(PackageConfigurations)','$(TestConfiguration)', '+')">
@@ -212,5 +257,6 @@
 <Property Name="TargetConfigurations" Value="Development" If="'$(TargetConfigurations)' == ''" />
 <Property Name="PackageConfigurations" Value="$(TargetConfigurations)" If="'$(PackageConfigurations)' == ''" />
 <Property Name="TestConfigurations" Value="$(PackageConfigurations)" If="'$(TestConfigurations)' == ''" />
+<Property Name="CookedEditorConfigurations" Value="Development" If="'$(CookedEditorConfigurations)' == ''" />
 
 </BuildGraph>
\ No newline at end of file
