Enable Cooked Editor in Project schedule

#rnx

diff --git a/Engine/Build/Graph/Tasks/Inc/PlatformOptions.xml b/Engine/Build/Graph/Tasks/Inc/PlatformOptions.xml
--- a/Engine/Build/Graph/Tasks/Inc/PlatformOptions.xml
+++ b/Engine/Build/Graph/Tasks/Inc/PlatformOptions.xml
@@ -51,7 +51,9 @@
 	<Option Name="WithMacEditor" DefaultValue="false" Description="Build and test Mac Editor"/>
 	<Option Name="WithLinuxEditor" DefaultValue="false" Description="Build and test Linux Editor"/>
 
-	<Option Name="WithCookedEditor" DefaultValue="false" Description="Build and cooked CookedEditor"/>
+	<Option Name="WithWin64CookedEditor" DefaultValue="false" Description="Build and cook Win64 CookedEditor"/>
+	<Option Name="WithMacCookedEditor" DefaultValue="false" Description="Build and cook Mac CookedEditor"/>
+	<Option Name="WithLinuxCookedEditor" DefaultValue="false" Description="Build and cook Linux CookedEditor"/>
 
 	<Option Name="WithAndroid" DefaultValue="false" Description="Build and test Android"/>
 	<Option Name="WithHoloLens" DefaultValue="false" Description="Build and test HoloLens"/>
@@ -78,9 +80,9 @@
 
 	<Do If="'$(CookedEditorPlatforms)' == ''">
 		<Property Name="CookedEditorPlatforms" Value="" />
-		<Property Name="CookedEditorPlatforms" Value="$(CookedEditorPlatforms)Win64+" If="$(WithWin64Editor) And $(WithCookedEditor)" />
-		<Property Name="CookedEditorPlatforms" Value="$(CookedEditorPlatforms)Mac+" If="$(WithMacEditor) And $(WithCookedEditor)" />
-		<Property Name="CookedEditorPlatforms" Value="$(CookedEditorPlatforms)Linux+" If="$(WithLinuxEditor) And $(WithCookedEditor)" />
+		<Property Name="CookedEditorPlatforms" Value="$(CookedEditorPlatforms)Win64+" If="$(WithWin64CookedEditor)" />
+		<Property Name="CookedEditorPlatforms" Value="$(CookedEditorPlatforms)Mac+" If="$(WithMacCookedEditor)" />
+		<Property Name="CookedEditorPlatforms" Value="$(CookedEditorPlatforms)Linux+" If="$(WithLinuxCookedEditor)" />
 	</Do>
 
 	<Do If="'$(TargetPlatforms)' == ''" >
