Add Shipping to the MediaFrameworkTest package projects schedule
+ workaround mobile limitation to only stage and package one configuration at a time

#rnx

diff --git a/Engine/Build/Graph/Tasks/BuildAndTestProject.xml b/Engine/Build/Graph/Tasks/BuildAndTestProject.xml
--- a/Engine/Build/Graph/Tasks/BuildAndTestProject.xml
+++ b/Engine/Build/Graph/Tasks/BuildAndTestProject.xml
@@ -19,6 +19,7 @@
 	<Option Name="MapIniSection" DefaultValue="" Restrict="" Description="MapIniSection argument to pass when cooking" />
 	<Option Name="ExtraCookArguments" DefaultValue="" Restrict="" Description="Extra arguments to use when cooking" />
 	<Option Name="ExtraStageAndPackageArguments" DefaultValue="" Restrict="" Description="Extra arguments to use when staging and packaging" />
+	<Option Name="ExtraPackageOnlyArguments" DefaultValue="" Restrict="" Description="Extra arguments to use when staging only" />	
 	<Option Name="TargetPlatformFilteredExtraStageAndPackageArguments" DefaultValue="" Restrict="" Description="Extra arguments to use when staging and packaging.  Filtered to specific platforms via TargetPlatformFilteredExtraStageAndPackageArgumentsFilter." />
 	<Option Name="TargetPlatformFilteredExtraStageAndPackageArgumentsFilter" DefaultValue="" Restrict="" Description="Platform filter applied to TargetPlatformFilteredExtraStageAndPackageArguments." />
 	<Option Name="WithRemoteExecWorkers" DefaultValue="false" Restrict="" Description="If your project uses standalone remote execution workers" />
@@ -352,6 +353,8 @@
 		<Property Name="TestAgentType" Value="$(HostAgentType)" />
 		<Property Name="TestAgentType" Value="Test$(HostAgentType)" If="'$(TargetPlatform)' == 'Win64' or '$(TargetPlatform)' == 'Mac'"/>
 
+		<Property Name="IsMobilePlatform" Value="'$(TargetPlatform)' == 'Android' or '$(TargetPlatform)' == 'IOS'"/>
+
 		<!-- Use incremental agents for building? -->
 		<Property Name="CompileAgentType" Value="$(HostAgentType)" />
 		<Property Name="CompileAgentType" Value="Incremental$(HostAgentType)" If="$(UseIncrementalAgents)" />
@@ -394,13 +397,15 @@
 		<Property Name="DoUnsyncHashing" Value="$(HostAgentType) == 'Win64'" />
 
 		<!-- Default BuildCookRun args needed when staging/packaging/archiving this configuration -->
-		<Property Name="BCRArgs" Value="-project=$(TargetProject) -platform=$(TargetPlatform) -cookflavor=&quot;$(BCRCookFlavor)&quot; -configuration=$(PackageConfigurations)" />
+		<Property Name="BCRArgs" Value="-project=$(TargetProject) -platform=$(TargetPlatform) -cookflavor=&quot;$(BCRCookFlavor)&quot;" />
 		<Property Name="BCRArgs" Value="$(BCRArgs) -client" If="'$(TargetType)' == 'Client'" />
 		<Property Name="BCRArgs" Value="$(BCRArgs) -server" If="'$(DedicatedServer)'" />
 
 		<!-- Support additional stage and package arguments that are only active on some platforms -->
 		<Property Name="StageAndPackageArguments" Value="$(ExtraStageAndPackageArguments)" />
 		<Property Name="StageAndPackageArguments" Value="$(StageAndPackageArguments) $(TargetPlatformFilteredExtraStageAndPackageArguments)" If="'$(TargetPlatformFilteredExtraStageAndPackageArgumentsFilter)' != '' and ContainsItem('$(TargetPlatformFilteredExtraStageAndPackageArgumentsFilter)', '$(TargetPlatform)', '+')" />
+		<Property Name="PackageOnlyArguments" Value="$(ExtraPackageOnlyArguments)"/>
+		<Property Name="PackageOnlyArguments" Value="$(PackageOnlyArguments) -archive -archivedirectory=&quot;$(ProjectOutputDirectory)/$(StagedPlatformFolder)&quot;" If="$(IsMobilePlatform)" />
 		
 		<!-- Build the client and if necessary a server-->
 		<Agent Name="Compile $(PlatformAgentName)" Type="$(CompileAgentType)">
@@ -414,7 +419,7 @@
 					</Do>
 				</ForEach>
 				<Do If="$(IsContentOnlyProject)">
-					<Command Name="BuildGame" Arguments="$(BCRArgs) -notools -ubtargs=&quot;-nobuilduht -BuildVersion=$(BuildName) $(GenericCompileArguments) $(ExtraTargetCompileArguments)&quot;" />
+					<Command Name="BuildGame" Arguments="$(BCRArgs) -configuration=$(TargetConfigurations) -notools -ubtargs=&quot;-nobuilduht -BuildVersion=$(BuildName) $(GenericCompileArguments) $(ExtraTargetCompileArguments)&quot;" />
 					<Tag Files="$(RootDir)/$(ProjectPath)/Binaries/..." With="#$(PlatformCompileNodeName) Binaries"/>
 				</Do>
 			</Node>
@@ -441,7 +446,16 @@
 			
 			<!-- Stage -->
 			<Node Name="$(PlatformStageNodeName)" Requires="$(PlatformCompileNodeName);#$(PlatformCompileNodeName) Binaries;#$(PlatformCookNodeName) Complete">
-				<Command Name="BuildCookRun" Arguments="$(BCRArgs) $(StageAndPackageArguments) -skipbuild -skipcook -stage -pak" />
+				<Switch>
+					<Case If="$(IsMobilePlatform)">
+						<ForEach Name="PackageConfiguration" Values="$(PackageConfigurations)" Separator="+">
+							<Command Name="BuildCookRun" Arguments="$(BCRArgs) -configuration=$(PackageConfiguration) $(StageAndPackageArguments) -skipbuild -skipcook -stage -pak" />
+						</ForEach>
+					</Case>
+					<Default>
+						<Command Name="BuildCookRun" Arguments="$(BCRArgs) -configuration=$(PackageConfigurations) $(StageAndPackageArguments) -skipbuild -skipcook -stage -pak" />
+					</Default>
+				</Switch>
 			</Node>
 			
 			<!-- Hash -->
@@ -459,12 +473,16 @@
 			
 			<!-- Package -->
 			<Node Name="$(PlatformPackageNodeName)" Requires="$(PlatformPublishStagedNodeName)">
-				<Command Name="BuildCookRun" Arguments="$(BCRArgs) $(StageAndPackageArguments) -skipbuild -skipcook -skipstage -skippak -package" />
+				<ForEach Name="PackageConfiguration" Values="$(PackageConfigurations)" Separator="+">
+					<Command Name="BuildCookRun" Arguments="$(BCRArgs) -configuration=$(PackageConfiguration) $(StageAndPackageArguments) -skipbuild -skipcook -skipstage -skippak -package $(PackageOnlyArguments)" />
+				</ForEach>
 			</Node>
 			
 			<!-- Publish (Packages) -->
 			<Node Name="$(PlatformPublishPackagedNodeName)" Requires="$(PlatformPackageNodeName)">
-				<Command Name="BuildCookRun" Arguments="$(BCRArgs) $(StageAndPackageArguments) -skipbuild -skipcook -skipstage -skippak -skippackage -archive -archivedirectory=&quot;$(NetworkOutputDirectory)/$(StagedPlatformFolder)/Packages&quot;" />
+				<ForEach Name="PackageConfiguration" Values="$(PackageConfigurations)" Separator="+">
+					<Command Name="BuildCookRun" Arguments="$(BCRArgs) -configuration=$(PackageConfiguration) $(StageAndPackageArguments) -skipbuild -skipcook -skipstage -skippak -skippackage -archive -archivedirectory=&quot;$(NetworkOutputDirectory)/$(StagedPlatformFolder)/Packages&quot;" />
+				</ForEach>
 			</Node>
 			
 		</Agent>
