Remove running with disableunity on incremental builds
#rnx

diff --git a/Engine/Build/Graph/Tasks/BuildAndTestProject.xml b/Engine/Build/Graph/Tasks/BuildAndTestProject.xml
--- a/Engine/Build/Graph/Tasks/BuildAndTestProject.xml
+++ b/Engine/Build/Graph/Tasks/BuildAndTestProject.xml
@@ -127,8 +127,8 @@
 	</Do>
 
 	<Do If="$(UseIncrementalAgents)" >
-		<Property Name="ExtraToolCompileArguments" Value="-allmodules -disableunity -nodebuginfo -2019 -Strict -WarningsAsErrors"/>
-		<Property Name="ExtraEditorCompileArguments" Value="-allmodules -disableunity -nodebuginfo -2019 -Strict -WarningsAsErrors"/>
+		<Property Name="ExtraToolCompileArguments" Value="-allmodules -nodebuginfo -2019 -Strict -WarningsAsErrors"/>
+		<Property Name="ExtraEditorCompileArguments" Value="-allmodules -nodebuginfo -2019 -Strict -WarningsAsErrors"/>
 		<Property Name="ExtraTargetCompileArguments" Value="-nodebuginfo -WarningsAsErrors"/>
 	</Do>
 
