Fix Cooked Editor exported path in Windows
+ set CookedEditor role in Gauntlet UE tests when -cookededitor switch is used

#rnx

diff --git a/Engine/Build/Graph/Tasks/BuildAndTestProject.xml b/Engine/Build/Graph/Tasks/BuildAndTestProject.xml
--- a/Engine/Build/Graph/Tasks/BuildAndTestProject.xml
+++ b/Engine/Build/Graph/Tasks/BuildAndTestProject.xml
@@ -257,6 +257,9 @@
 		<Property Name="TestAgentType" Value="Test$(EditorPlatform)" />
 		<Property Name="TestAgentType" Value="$(TestAgentOverride)" If="'$(TestAgentOverride)' != ''"/>
 
+		<Property Name="EditorPlatformPath" Value="$(EditorPlatform)" />
+		<Property Name="EditorPlatformPath" Value="Windows" If="'$(EditorPlatform)' == 'Win64'" />
+
 		<Property Name="PlatformPublishNodeName" Value="Publish CookedEditor $(EditorPlatform)" />
 
 		<Property Name="CookedEditorCookRequirements" Value="Compile Tools $(EditorPlatform);Compile $(ProjectName)Editor $(EditorPlatform);Compile $(ProjectName)CookedEditor $(EditorPlatform)"/>
@@ -279,9 +282,9 @@
 
 			<!-- Publish CookedEditor -->
 			<Node Name="Publish CookedEditor $(EditorPlatform)" Requires="Stage CookedEditor $(EditorPlatform)">
-				<Log Message="Copying to $(NetworkOutputDirectory)/$(EditorPlatform)CookedEditor"/>
+				<Log Message="Copying to $(NetworkOutputDirectory)/$(EditorPlatformPath)CookedEditor"/>
 				<Do If="!$(SkipPublish)">
-					<Copy Files="..." From="$(ProjectOutputDirectory)/$(EditorPlatform)CookedEditor" To="$(NetworkOutputDirectory)/$(EditorPlatform)CookedEditor"/>
+					<Copy Files="..." From="$(ProjectOutputDirectory)/$(EditorPlatformPath)CookedEditor" To="$(NetworkOutputDirectory)/$(EditorPlatformPath)CookedEditor"/>
 				</Do>
 			</Node>
 
@@ -305,7 +308,7 @@
 				<Property Name="GauntletReportArgs" Value="$(GauntletReportArgs) -ReportURL=&quot;$(AutomationReportUrl)/$(OutputFragmentPath)/&quot;" If="'$(AutomationReportUrl)' != ''" />
 			</Do>
 			<Node Name="$(PlatformBootTestNodeName)" Requires="$(TestPrerequisites)">
-				<Command Name="RunUnreal" Arguments="-project=$(TargetProject) -build=&quot;$(NetworkOutputDirectory)/$(EditorPlatform)CookedEditor&quot; -cookededitor -test=&quot;UE.EditorBootTest&quot; $(GauntletEditorTestArgs) $(GauntletReportArgs)" />
+				<Command Name="RunUnreal" Arguments="-project=$(TargetProject) -build=&quot;$(NetworkOutputDirectory)/$(EditorPlatformPath)CookedEditor&quot; -test=&quot;UE.EditorBootTest&quot; $(GauntletCookedEditorTestArgs) $(GauntletReportArgs)" />
 			</Node>
 
 			<!-- Link the boot test in to the prerequisites and generated nodes if its enabled-->
@@ -323,7 +326,7 @@
 				</Do>
 				<Property Name="TestNodeName" Value="CookedEditor $(EditorPlatform) Test=$(TestName)" />
 				<Node Name="$(TestNodeName)" Requires="$(TestPrerequisites)">
-					<Command Name="RunUnreal" Arguments="-project=$(TargetProject) -build=&quot;$(NetworkOutputDirectory)/$(EditorPlatform)CookedEditor&quot; -cookededitor -test=&quot;$(TestName)&quot; $(GauntletReportArgs) $(GauntletEditorTestArgs)" />
+					<Command Name="RunUnreal" Arguments="-project=$(TargetProject) -build=&quot;$(NetworkOutputDirectory)/$(EditorPlatformPath)CookedEditor&quot; -test=&quot;$(TestName)&quot; $(GauntletReportArgs) $(GauntletCookedEditorTestArgs)" />
 				</Node>
 				<!-- Add each test node to our generated list -->
 				<Property Name="PlatformCookedEditorTestNodes" Value="$(PlatformCookedEditorTestNodes);$(TestNodeName)" />
diff --git a/Engine/Build/Graph/Tasks/Inc/GauntletSettings.xml b/Engine/Build/Graph/Tasks/Inc/GauntletSettings.xml
--- a/Engine/Build/Graph/Tasks/Inc/GauntletSettings.xml
+++ b/Engine/Build/Graph/Tasks/Inc/GauntletSettings.xml
@@ -36,6 +36,7 @@
 	<Property Name="GauntletTargetTestArgs" Value="$(GauntletCommonArgs) $(GauntletDeviceArgs) $(ExtraGauntletTargetTestArgs)" />
 	<Property Name="GauntletTargetTestAltArgs" Value="$(GauntletCommonArgs) $(GauntletDeviceAltArgs) $(ExtraGauntletTargetTestArgs)" />
 	<Property Name="GauntletEditorTestArgs" Value="-build=editor $(GauntletCommonArgs) $(ExtraGauntletEditorTestArgs)" />
+	<Property Name="GauntletCookedEditorTestArgs" Value="-cookededitor $(GauntletCommonArgs) $(ExtraGauntletEditorTestArgs)" />
 
     <!-- this is fully generated dynamiclly when needed, it's simply here so it can be used without worrying if its set -->
     <Property Name="GauntletReportArgs" Value="" />
diff --git a/Engine/Source/Programs/AutomationTool/Gauntlet/Unreal/Automation/UE.Automation.cs b/Engine/Source/Programs/AutomationTool/Gauntlet/Unreal/Automation/UE.Automation.cs
--- a/Engine/Source/Programs/AutomationTool/Gauntlet/Unreal/Automation/UE.Automation.cs
+++ b/Engine/Source/Programs/AutomationTool/Gauntlet/Unreal/Automation/UE.Automation.cs
@@ -334,7 +334,7 @@ namespace UE
 			AutomationTestConfig Config = base.GetConfiguration();
 
 			// Tests in the editor only require a single role
-			UnrealTestRole EditorRole = Config.RequireRole(UnrealTargetRole.Editor);
+			UnrealTestRole EditorRole = Config.RequireRole(Config.CookedEditor ? UnrealTargetRole.CookedEditor : UnrealTargetRole.Editor);
 			EditorRole.CommandLineParams.AddRawCommandline("-NoWatchdog -stdout -FORCELOGFLUSH -CrashForUAT -log");
 
 			return Config;
diff --git a/Engine/Source/Programs/AutomationTool/Gauntlet/Unreal/Automation/UE.BootTest.cs b/Engine/Source/Programs/AutomationTool/Gauntlet/Unreal/Automation/UE.BootTest.cs
--- a/Engine/Source/Programs/AutomationTool/Gauntlet/Unreal/Automation/UE.BootTest.cs
+++ b/Engine/Source/Programs/AutomationTool/Gauntlet/Unreal/Automation/UE.BootTest.cs
@@ -214,7 +214,7 @@ namespace UE
 			UnrealTestConfiguration Config = base.GetConfiguration();
 			// currently needed as BootTest isn't an abstract class. Can be changed for 4.27
 			Config.ClearRoles();
-			UnrealTestRole EditorRole = Config.RequireRole(UnrealTargetRole.Editor);
+			UnrealTestRole EditorRole = Config.RequireRole(Config.CookedEditor ? UnrealTargetRole.CookedEditor : UnrealTargetRole.Editor);
 			EditorRole.CommandLineParams.Add("execcmds", "QUIT_EDITOR");
 			return Config;
 		}
