Avoid compiling twice the project with content only project

#rnx

diff --git a/Engine/Build/Graph/Tasks/BuildAndTestProject.xml b/Engine/Build/Graph/Tasks/BuildAndTestProject.xml
--- a/Engine/Build/Graph/Tasks/BuildAndTestProject.xml
+++ b/Engine/Build/Graph/Tasks/BuildAndTestProject.xml
@@ -415,16 +415,20 @@
 
 			<!-- Compile -->
 			<Node Name="$(PlatformCompileNodeName)" Requires="$(PlatformToolsNodeName)" Produces="#$(PlatformCompileNodeName) Binaries">
-				<ForEach Name="TargetConfiguration" Values="$(TargetConfigurations)" Separator="+">
-					<Compile Target="$(TargetName)" Platform="$(TargetPlatform)" Configuration="$(TargetConfiguration)" Arguments="-BuildVersion=&quot;$(BuildName)&quot; $(GenericCompileArguments) $(ExtraTargetCompileArguments)"/>
-					<Do If="'$(DedicatedServer)'" >
-						<Compile Target="$(TargetName)Server" Platform="Win64" Configuration="$(TargetConfiguration)" Arguments="-BuildVersion=&quot;$(BuildName)&quot; $(GenericCompileArguments) $(ExtraTargetCompileArguments)"/>
-					</Do>
-				</ForEach>
-				<Do If="$(IsContentOnlyProject)">
-					<Command Name="BuildGame" Arguments="$(BCRArgs) -configuration=$(TargetConfigurations) -notools -ubtargs=&quot;-nobuilduht -BuildVersion=$(BuildName) $(GenericCompileArguments) $(ExtraTargetCompileArguments)&quot;" />
-					<Tag Files="$(RootDir)/$(ProjectPath)/Binaries/..." With="#$(PlatformCompileNodeName) Binaries"/>
-				</Do>
+				<Switch>
+					<Case If="$(IsContentOnlyProject)">
+						<Command Name="BuildGame" Arguments="$(BCRArgs) -configuration=$(TargetConfigurations) -notools -ubtargs=&quot;-nobuilduht -BuildVersion=$(BuildName) $(GenericCompileArguments) $(ExtraTargetCompileArguments)&quot;" />
+						<Tag Files="$(RootDir)/$(ProjectPath)/Binaries/..." With="#$(PlatformCompileNodeName) Binaries"/>
+					</Case>
+					<Default>
+						<ForEach Name="TargetConfiguration" Values="$(TargetConfigurations)" Separator="+">
+							<Compile Target="$(TargetName)" Platform="$(TargetPlatform)" Configuration="$(TargetConfiguration)" Arguments="-BuildVersion=&quot;$(BuildName)&quot; $(GenericCompileArguments) $(ExtraTargetCompileArguments)"/>
+							<Do If="'$(DedicatedServer)'" >
+								<Compile Target="$(TargetName)Server" Platform="Win64" Configuration="$(TargetConfiguration)" Arguments="-BuildVersion=&quot;$(BuildName)&quot; $(GenericCompileArguments) $(ExtraTargetCompileArguments)"/>
+							</Do>
+						</ForEach>
+					</Default>
+				</Switch>
 			</Node>
 
 			<!-- (Optional) Upload Symbols -->
