Fixing incorrect paths for Linux Test agents

#rnx

diff --git a/Engine/Build/Graph/Tasks/Inc/PathSettings.xml b/Engine/Build/Graph/Tasks/Inc/PathSettings.xml
--- a/Engine/Build/Graph/Tasks/Inc/PathSettings.xml
+++ b/Engine/Build/Graph/Tasks/Inc/PathSettings.xml
@@ -45,14 +45,17 @@
 
 	<!-- Scratch network path that can be used for tests if we aren't publishing. Assumes there's a cleanup script.. -->
 	<Property Name="NetworkTempRoot" Value="/Volumes/Builds/UE5/TempStorage"/>
+	<Property Name="NetworkReportRoot" Value="/mnt/Builds/UE5/TempStorage" If="'$(HostPlatform)' == 'Linux'"/>
 	<Property Name="NetworkTempRoot" Value="P:/Builds/UE5/TempStorage" If="'$(HostPlatform)' == 'Win64'"/>
 
 	<!-- Network path that can be used for builds that are published -->
 	<Property Name="NetworkPublishRoot" Value="/Volumes/Builds/UE5/PackagedBuilds"/>
+	<Property Name="NetworkReportRoot" Value="/mnt/Builds/UE5/PackagedBuilds" If="'$(HostPlatform)' == 'Linux'"/>
 	<Property Name="NetworkPublishRoot" Value="P:/Builds/UE5/PackagedBuilds" If="'$(HostPlatform)' == 'Win64'"/>
 
 	<!-- Network path that can be used for storing reports and logs -->
 	<Property Name="NetworkReportRoot" Value="/Volumes/Builds/Automation/Reports"/>
+	<Property Name="NetworkReportRoot" Value="/mnt/Builds/Automation/Reports" If="'$(HostPlatform)' == 'Linux'"/>
 	<Property Name="NetworkReportRoot" Value="P:/Builds/Automation/Reports" If="'$(HostPlatform)' == 'Win64'"/>
 
 	<!-- Allow scripts to override the network locations-->
