Add an option to compile UnrealInsights as a requirement

#rnx

diff --git a/Engine/Build/Graph/Tasks/BuildAndTestProject.xml b/Engine/Build/Graph/Tasks/BuildAndTestProject.xml
--- a/Engine/Build/Graph/Tasks/BuildAndTestProject.xml
+++ b/Engine/Build/Graph/Tasks/BuildAndTestProject.xml
@@ -28,6 +28,7 @@
 	<Option Name="AgentOverride" DefaultValue="" Description="If set, other logic is ignored and this agent type is used for all non-test work." />
 	<Option Name="TestAgentOverride" DefaultValue="" Description="If set, other logic is ignored and this agent type is used for all test work." />
 	<Option Name="ClientCompileAgentOverride" DefaultValue="" Description="Specifies the agent type to use for C++ compilation jobs for target platform binaries." />
+	<Option Name="IncludeUnrealInsightsInToolCompile" DefaultValue="false" Description="Include Unreal Insights into Compile Tools prerequisites" />
 
 	<!-- 
 		Allow project graphs that Include this file to specify their own defaults.
@@ -182,6 +183,7 @@
 						</ForEach>
 					</Do>
 				</Do>
+				<Compile Target="UnrealInsights" Platform="$(EditorPlatform)" Configuration="Shipping" Arguments="$(GenericCompileArguments)" If="$(IncludeUnrealInsightsInToolCompile)"/>
 			</Node>
 			<Node Name="Compile $(ProjectName)Editor $(EditorPlatform)" Requires="Compile Tools $(EditorPlatform)">
 				<Compile Target="$(EditorTargetName)" Platform="$(EditorPlatform)" Configuration="Development" Arguments="$(GenericCompileArguments) $(ExtraEditorCompileArguments)"/>
