Exclude DotNet folder when Tagging ContentOnlyProject compile output

#rnx

diff --git a/Engine/Build/Graph/Tasks/BuildAndTestProject.xml b/Engine/Build/Graph/Tasks/BuildAndTestProject.xml
--- a/Engine/Build/Graph/Tasks/BuildAndTestProject.xml
+++ b/Engine/Build/Graph/Tasks/BuildAndTestProject.xml
@@ -420,7 +420,7 @@
 				<Switch>
 					<Case If="$(IsContentOnlyProject)">
 						<Command Name="BuildGame" Arguments="$(BCRArgs) -configuration=$(TargetConfigurations) -notools -ubtargs=&quot;-nobuilduht -BuildVersion=$(BuildName) $(GenericCompileArguments) $(ExtraTargetCompileArguments)&quot;" />
-						<Tag Files="$(RootDir)/$(ProjectPath)/Binaries/..." With="#$(PlatformCompileNodeName) Binaries"/>
+						<Tag Files="$(RootDir)/$(ProjectPath)/Binaries/..." Except="$(ProjectPath)/Binaries/DotNET/..." With="#$(PlatformCompileNodeName) Binaries"/>
 					</Case>
 					<Default>
 						<ForEach Name="TargetConfiguration" Values="$(TargetConfigurations)" Separator="+">
