Add argument to setup Module Role
+ Enforce using a the CustomModuleRoles argument with CookedEditor in BuildAndTestProject buildgraph

#rnx

diff --git a/Engine/Build/Graph/Tasks/BuildAndTestProject.xml b/Engine/Build/Graph/Tasks/BuildAndTestProject.xml
--- a/Engine/Build/Graph/Tasks/BuildAndTestProject.xml
+++ b/Engine/Build/Graph/Tasks/BuildAndTestProject.xml
@@ -298,6 +298,7 @@
 
 		<!-- Tests -->
 		<Agent Name="Test $(ProjectName)CookedEditor $(EditorPlatform)" Type="$(TestAgentType)">
+			<Property Name="CustomModuleRoleArg" Value="-CustomModuleRoles=$(CookedEditorTargetName):CookedEditor"/>
 			<!-- Prerequisites for the editor test to run -->
 			<Property Name="TestPrerequisites" Value="$(GeneratedCookedEditorNodes)"/>
 			<Property Name="TestPrerequisites" Value="$(TestPrerequisites);$(PlatformPublishNodeName)" If="!$(PublishBuilds)"/>
@@ -310,7 +311,7 @@
 				<Property Name="GauntletReportArgs" Value="$(GauntletReportArgs) -ReportURL=&quot;$(AutomationReportUrl)/$(OutputFragmentPath)/&quot;" If="'$(AutomationReportUrl)' != ''" />
 			</Do>
 			<Node Name="$(PlatformBootTestNodeName)" Requires="$(TestPrerequisites)">
-				<Command Name="RunUnreal" Arguments="-project=$(TargetProject) -build=&quot;$(NetworkOutputDirectory)/$(EditorPlatformPath)CookedEditor&quot; -test=&quot;UE.EditorBootTest&quot; $(GauntletCookedEditorTestArgs) $(GauntletReportArgs)" />
+				<Command Name="RunUnreal" Arguments="-project=$(TargetProject) -build=&quot;$(NetworkOutputDirectory)/$(EditorPlatformPath)CookedEditor&quot; -test=&quot;UE.EditorBootTest&quot; $(GauntletCookedEditorTestArgs) $(CustomModuleRoleArg) $(GauntletReportArgs)" />
 			</Node>
 
 			<!-- Link the boot test in to the prerequisites and generated nodes if its enabled-->
@@ -328,7 +329,7 @@
 				</Do>
 				<Property Name="TestNodeName" Value="CookedEditor $(EditorPlatform) Test=$(TestName)" />
 				<Node Name="$(TestNodeName)" Requires="$(TestPrerequisites)">
-					<Command Name="RunUnreal" Arguments="-project=$(TargetProject) -build=&quot;$(NetworkOutputDirectory)/$(EditorPlatformPath)CookedEditor&quot; -test=&quot;$(TestName)&quot; $(GauntletReportArgs) $(GauntletCookedEditorTestArgs)" />
+					<Command Name="RunUnreal" Arguments="-project=$(TargetProject) -build=&quot;$(NetworkOutputDirectory)/$(EditorPlatformPath)CookedEditor&quot; -test=&quot;$(TestName)&quot; $(GauntletReportArgs) $(CustomModuleRoleArg) $(GauntletCookedEditorTestArgs)" />
 				</Node>
 				<!-- Add each test node to our generated list -->
 				<Property Name="PlatformCookedEditorTestNodes" Value="$(PlatformCookedEditorTestNodes);$(TestNodeName)" />
diff --git a/Engine/Source/Programs/AutomationTool/Gauntlet/Unreal/Base/Gauntlet.UnrealTestContext.cs b/Engine/Source/Programs/AutomationTool/Gauntlet/Unreal/Base/Gauntlet.UnrealTestContext.cs
--- a/Engine/Source/Programs/AutomationTool/Gauntlet/Unreal/Base/Gauntlet.UnrealTestContext.cs
+++ b/Engine/Source/Programs/AutomationTool/Gauntlet/Unreal/Base/Gauntlet.UnrealTestContext.cs
@@ -189,6 +189,13 @@ namespace Gauntlet
 		[AutoParam("")]
 		public string ArtifactPostfix;
 
+		/// <summary>
+		/// Add custom module name/role pair to identify custom project target executable
+		/// Format is <name>:<role>+<name>:<role>
+		/// </summary>
+		[AutoParam("")]
+		public string CustomModuleRoles { get; set; }
+
 		/// <summary>
 		/// Less logging
 		/// </summary>
@@ -343,6 +350,32 @@ namespace Gauntlet
 					&& !Arg.StartsWith("device=", StringComparison.OrdinalIgnoreCase))
 				.ToArray();
 			Params = new Params(CleanArgs);
+
+			// Custom Module name/role
+			if (!string.IsNullOrEmpty(CustomModuleRoles))
+			{
+				foreach (var Pair in CustomModuleRoles.Split("+"))
+				{
+					var SplittedPair = Pair.Split(":");
+					if (SplittedPair.Length > 1)
+					{
+						string Name = SplittedPair[0];
+						UnrealTargetRole Role = UnrealTargetRole.Unknown;
+						if (Enum.TryParse(SplittedPair[1], out Role))
+						{
+							UnrealHelpers.AddCustomModuleName(Name, Role);
+						}
+						else
+						{
+							throw new AutomationException(string.Format("Target Role '{0}' for Custom Module '{1}' is unknown.", SplittedPair[1], Name));
+						}
+					}
+					else
+					{
+						Gauntlet.Log.Warning("CustomModuleRoles is poorly formatted. Expected <name>:<role> pair. Got '{0}'", Pair);
+					}
+				}
+			}
 		}
 	}
 	
