Use DevkitAutomation for agents that deploy on console

#rnx

diff --git a/Engine/Build/Graph/Tasks/BuildAndTestProject.xml b/Engine/Build/Graph/Tasks/BuildAndTestProject.xml
--- a/Engine/Build/Graph/Tasks/BuildAndTestProject.xml
+++ b/Engine/Build/Graph/Tasks/BuildAndTestProject.xml
@@ -27,6 +27,7 @@
 	<Option Name="UseIncrementalAgents" DefaultValue="false" Restrict="true|false" Description="Use incremental agents for building" />
 	<Option Name="AgentOverride" DefaultValue="" Description="If set, other logic is ignored and this agent type is used for all non-test work." />
 	<Option Name="TestAgentOverride" DefaultValue="" Description="If set, other logic is ignored and this agent type is used for all test work." />
+	<Option Name="TargetPlatformTestAgentOverride" DefaultValue="DevkitAutomation" Description="Override agent type use to run target platform tests"/>
 	<Option Name="ClientCompileAgentOverride" DefaultValue="" Description="Specifies the agent type to use for C++ compilation jobs for target platform binaries." />
 	<Option Name="IncludeUnrealInsightsInToolCompile" DefaultValue="false" Description="Include Unreal Insights into Compile Tools prerequisites" />
 
@@ -355,9 +356,9 @@
 		<!-- Declare type of host required to build & cook this platform-->
 		<Property Name="HostAgentType" Value="Win64" />
 		<Property Name="HostAgentType" Value="Mac" If="'$(TargetPlatform)' == 'Mac' or '$(TargetPlatform)' == 'IOS' or '$(TargetPlatform)' == 'tvOS'"/>
-		<!-- On consoles we can use regualr agents, but locally hosted tests need a Test agent-->
-		<Property Name="TestAgentType" Value="$(HostAgentType)" />
-		<Property Name="TestAgentType" Value="Test$(HostAgentType)" If="'$(TargetPlatform)' == 'Win64' or '$(TargetPlatform)' == 'Mac'"/>
+		<!-- On consoles we can use agents with devkit setup, but locally hosted tests need a Test agent-->
+		<Property Name="TestAgentType" Value="$(TargetPlatformTestAgentOverride)" />
+		<Property Name="TestAgentType" Value="Test$(HostAgentType)" If="'$(TargetPlatform)' == 'Win64' or '$(TargetPlatform)' == 'Mac' or '$(TargetPlatform)' == 'Linux'"/>
 
 		<Property Name="IsMobilePlatform" Value="'$(TargetPlatform)' == 'Android' or '$(TargetPlatform)' == 'IOS'"/>
 
