Fix depencies for target linux tests

#rnx

diff --git a/Engine/Build/Graph/Tasks/BuildAndTestProject.xml b/Engine/Build/Graph/Tasks/BuildAndTestProject.xml
--- a/Engine/Build/Graph/Tasks/BuildAndTestProject.xml
+++ b/Engine/Build/Graph/Tasks/BuildAndTestProject.xml
@@ -356,6 +356,7 @@
 		<!-- Declare type of host required to build & cook this platform-->
 		<Property Name="HostAgentType" Value="Win64" />
 		<Property Name="HostAgentType" Value="Mac" If="'$(TargetPlatform)' == 'Mac' or '$(TargetPlatform)' == 'IOS' or '$(TargetPlatform)' == 'tvOS'"/>
+		<Property Name="HostAgentType" Value="Linux" If="'$(TargetPlatform)' == 'Linux'"/>
 		<!-- On consoles we can use agents with devkit setup, but locally hosted tests need a Test agent-->
 		<Property Name="TestAgentType" Value="$(TargetPlatformTestAgentOverride)" />
 		<Property Name="TestAgentType" Value="Test$(HostAgentType)" If="'$(TargetPlatform)' == 'Win64' or '$(TargetPlatform)' == 'Mac' or '$(TargetPlatform)' == 'Linux'"/>
