Fix reference to Linux depencies

#rnx

diff --git a/Engine/Build/Graph/Tasks/Inc/PlatformOptions.xml b/Engine/Build/Graph/Tasks/Inc/PlatformOptions.xml
--- a/Engine/Build/Graph/Tasks/Inc/PlatformOptions.xml
+++ b/Engine/Build/Graph/Tasks/Inc/PlatformOptions.xml
@@ -150,6 +150,7 @@
 	<!-- Declare type of host required to build & cook this platform-->
 	<Property Name="HostAgentType" Value="Win64" />
 	<Property Name="HostAgentType" Value="Mac" If="'$(TargetPlatform)' == 'Mac' or '$(TargetPlatform)' == 'IOS' or '$(TargetPlatform)' == 'tvOS'"/>
+	<Property Name="HostAgentType" Value="Linux" If="'$(TargetPlatform)' == 'Linux'"/>
 
 	<Do If="!ContainsItem('$(RequiredEditorPlatforms)','$(HostAgentType)', '+')">
 		<Property Name="RequiredEditorPlatforms" Value="$(RequiredEditorPlatforms)+$(HostAgentType)" />
