Horde: Add a job type for running the local executor in the default config file.

diff --git a/Engine/Build/Graph/Examples/TagsAndFiles.xml b/Engine/Build/Graph/Examples/TagsAndFiles.xml
--- a/Engine/Build/Graph/Examples/TagsAndFiles.xml
+++ b/Engine/Build/Graph/Examples/TagsAndFiles.xml
@@ -1,7 +1,7 @@
 <?xml version='1.0' ?>
 <BuildGraph xmlns="http://www.epicgames.com/BuildGraph" xmlns:xsi="http://www.w3.org/2001/XMLSchema-instance" xsi:schemaLocation="http://www.epicgames.com/BuildGraph ../Schema.xsd" >
 
-	<Agent Name="Agent 1" Type="CompileWin64">
+	<Agent Name="Agent 1" Type="AnyAgent">
 		<Node Name="Manipulate Some Files">
 			<Log Message="Here are all the files in Engine/Build/..." Files="Engine/Build/..."/>
 
@@ -26,7 +26,7 @@
 
 		</Node>
 	</Agent>
-	<Agent Name="Agent 2" Type="CompileWin64">
+	<Agent Name="Agent 2" Type="AnyAgent">
 		<Node Name="First Downstream Node" Requires="Manipulate Some Files">
 			<!-- Test that we have all the side-effects of tests 1-3 from the 'Manipulate Some Files' node -->
 			<Log Message="Here are all the files in Engine/TestNumber1/..." Files="Engine/TestNumber1/..."/>
@@ -48,7 +48,7 @@
 			<Log Message="Produced files:" Files="#OutputForTestNumber5"/>
 		</Node>
 	</Agent>
-	<Agent Name="Agent 3" Type="CompileWin64">
+	<Agent Name="Agent 3" Type="AnyAgent">
 		<Node Name="Consume a Single Named Output" Requires="#OutputForTestNumber5">
 			<!-- Show that we ONLY have the #OutputForTestNumber5 output -->
 			<Log Message="Here are all the files in Engine/TestNumber1/..." Files="Engine/TestNumber1/..."/>
@@ -58,7 +58,7 @@
 			<Log Message="Here are all the files in Engine/TestNumber5/..." Files="Engine/TestNumber5/..."/>
 		</Node>
 	</Agent>
-	<Agent Name="Agent 4" Type="CompileWin64">
+	<Agent Name="Agent 4" Type="AnyAgent">
 		<Node Name="Create Three Named Outputs" Produces="#First Named Output;#Second Named Output;#Third Named Output">
 			<!-- Start again by creating three named outputs, two containing different files and one containing files from each -->
 			<Copy From="Engine/Build" Files="*" To="Engine/TestNumber5" Tag="#First Named Output;#Third Named Output"/>
@@ -71,14 +71,14 @@
 			<Log Message="Created '#Third Named Output' with:" Files="#Third Named Output"/>
 		</Node>
 	</Agent>
-	<Agent Name="Agent 5" Type="CompileWin64">
+	<Agent Name="Agent 5" Type="AnyAgent">
 		<Node Name="Add Additional Files" Requires="#Second Named Output" Produces="#Fourth Named Output">
 			<!-- Add some more files to '#Second Named Output' -->
 			<Tag Files="#Second Named Output" With="#Fourth Named Output"/>
 			<Copy From="Engine/Build" Files="*" To="Engine/TestNumber6/More" Tag="#Fourth Named Output"/>
 		</Node>
 	</Agent>
-	<Agent Name="Agent 6" Type="CompileWin64">
+	<Agent Name="Agent 6" Type="AnyAgent">
 		<Node Name="List Named Outputs" Requires="Create Three Named Outputs;Add Additional Files">
 			<!-- List all the outputs. -->
 			<Log Message="'#Create Three Named Outputs' contains:" Files="#Create Three Named Outputs"/>
@@ -89,7 +89,7 @@
 			<Log Message="'#Fourth Named Output' contains:" Files="#Fourth Named Output"/>
 		</Node>
 	</Agent>
-	<Agent Name="Agent 7" Type="CompileWin64">
+	<Agent Name="Agent 7" Type="AnyAgent">
 		<Node Name="Copy Syntax">
 			<Copy Files="..." From="Engine/Build" To="Engine/TestNumber6/Output0" />
 			<Copy From="Engine/Build/Build.version" To="Engine/TestNumber6/Output1/Output.version" /> <!-- single file -->
@@ -103,7 +103,7 @@
 		</Node>
 	</Agent>
 
-	<Agent Name="Trigger Group" Type="CompileWin64">
+	<Agent Name="Trigger Group" Type="AnyAgent">
 		<Node Name="Print all Output Folders" Requires="Manipulate Some Files;Create A Named Output;Copy Syntax">
 			<!-- Test that we have all the side-effects of the upstream node -->
 			<Log Message="Here are all the files in Engine/TestNumber1/..." Files="Engine/TestNumber1/..."/>
diff --git a/Engine/Source/Programs/Horde/Horde.Build/Defaults/ue5-main.stream.json b/Engine/Source/Programs/Horde/Horde.Build/Defaults/ue5-main.stream.json
--- a/Engine/Source/Programs/Horde/Horde.Build/Defaults/ue5-main.stream.json
+++ b/Engine/Source/Programs/Horde/Horde.Build/Defaults/ue5-main.stream.json
@@ -9,6 +9,7 @@
 			"templates":
 			[
 				"test-executor",
+				"local-executor",
 				"custom"
 			],
 			"columns":
@@ -58,6 +59,16 @@
 				"-Target=Full Build"
 			]
 		},
+		{
+			"id": "local-executor",
+			"name": "Run Local Executor",
+			"executor": "Local",
+			"initialAgentType": "AnyAgent",
+			"arguments": [
+				"-Script=Engine/Build/Graph/Examples/TagsAndFiles.xml",
+				"-Target=Tag and File Examples"
+			]
+		},
 		{
 			"id": "custom",
 			"name": "Custom...",
