First pass replay-based PGO Profiling build graph macro.

 - includes platform-specific PGO properties for all supported platforms.
 - stubs for Android & Linux although they are not fully functional yet.

#rnx

diff --git a/Engine/Build/Graph/Tasks/Inc/Android/PGOProperties.xml b/Engine/Build/Graph/Tasks/Inc/Android/PGOProperties.xml
new file mode 100644
--- /dev/null
+++ b/Engine/Build/Graph/Tasks/Inc/Android/PGOProperties.xml
@@ -0,0 +1,20 @@
+<?xml version='1.0' ?>
+<BuildGraph xmlns="http://www.epicgames.com/BuildGraph" xmlns:xsi="http://www.w3.org/2001/XMLSchema-instance" xsi:schemaLocation="http://www.epicgames.com/BuildGraph ../../Schema.xsd" >
+
+	<!-- add us to the platform list-->
+	<!-- NOTE: Android not currently supported - it requires a packaged build & the general-purpose node doesn't support this yet
+	<Property Name="AllPGOPlatforms" Value="$(AllPGOPlatforms);Android" If="ContainsItem('$(TargetPlatforms)', 'Android', '+')"/>
+	-->
+
+	<!-- set platform-specific properties -->
+	<Property Name="PGOBinaryStagingPathAndroid" Value="$(ProjectName)/Binaries/Android/"/>
+	<Property Name="PGOProfileOutputPathAndroid" Value="$(ProjectPath)/Platforms/Android/Build/PGO"/>
+	<Property Name="PGOProfileOutputFilterAndroid" Value="*.profdata"/>
+	<Property Name="PGOProfileCommandLineAndroid" Value=""/>
+	<Property Name="PGOCompileArgsAndroid" Value=""/>
+	<Property Name="PGOCompileExtraBinariesAndroid" Value=""/>
+
+	<!-- no special preprocessing required for this platform -->
+	<Macro Name="PGOPrepareBuildForProfilingAndroid" Arguments="LocalStagingDir;Configuration"/>
+
+</BuildGraph>
\ No newline at end of file
diff --git a/Engine/Build/Graph/Tasks/Inc/Linux/PGOProperties.xml b/Engine/Build/Graph/Tasks/Inc/Linux/PGOProperties.xml
new file mode 100644
--- /dev/null
+++ b/Engine/Build/Graph/Tasks/Inc/Linux/PGOProperties.xml
@@ -0,0 +1,21 @@
+<?xml version='1.0' ?>
+<BuildGraph xmlns="http://www.epicgames.com/BuildGraph" xmlns:xsi="http://www.w3.org/2001/XMLSchema-instance" xsi:schemaLocation="http://www.epicgames.com/BuildGraph ../../Schema.xsd" >
+
+	<!-- add us to the platform list-->
+	<!-- Linux not currently supported. Gauntlet PGO automation has not been implemented.
+	<Property Name="AllPGOPlatforms" Value="$(AllPGOPlatforms);Linux" If="ContainsItem('$(TargetPlatforms)', 'Linux', '+')"/>
+	-->
+
+	<!-- set platform-specific properties -->
+	<Property Name="PGOBinaryStagingPathLinux" Value="$(ProjectName)/Binaries/Linux/"/>
+	<Property Name="PGOProfileOutputLinux" Value="$(ProjectPath)/Build/Linux/PGO"/>
+	<Property Name="PGOProfileOutputFilterLinux" Value="*.profdata"/>
+	<Property Name="PGOProfileCommandLineLinux" Value=""/>
+	<Property Name="PGOCompileArgsLinux" Value=""/>
+	<Property Name="PGOCompileExtraBinariesLinux" Value=""/>
+
+	<!-- no special preprocessing required for this platform -->
+	<Macro Name="PGOPrepareBuildForProfilingLinux" Arguments="LocalStagingDir;Configuration"/>
+
+
+</BuildGraph>
\ No newline at end of file
diff --git a/Engine/Build/Graph/Tasks/Inc/Win64/PGOProperties.xml b/Engine/Build/Graph/Tasks/Inc/Win64/PGOProperties.xml
new file mode 100644
--- /dev/null
+++ b/Engine/Build/Graph/Tasks/Inc/Win64/PGOProperties.xml
@@ -0,0 +1,22 @@
+<?xml version='1.0' ?>
+<BuildGraph xmlns="http://www.epicgames.com/BuildGraph" xmlns:xsi="http://www.w3.org/2001/XMLSchema-instance" xsi:schemaLocation="http://www.epicgames.com/BuildGraph ../../Schema.xsd" >
+
+	<!-- add us to the platform list-->
+	<Property Name="AllPGOPlatforms" Value="$(AllPGOPlatforms);Win64" If="ContainsItem('$(TargetPlatforms)', 'Win64', '+')"/>
+
+	<!-- set platform-specific properties -->
+	<Property Name="PGOBinaryStagingPathWin64" Value="$(ProjectName)/Binaries/Win64/"/>
+	<Property Name="PGOProfileOutputWin64" Value="$(ProjectPath)/Platforms/Windows/Build/PGO"/>
+	<Property Name="PGOProfileOutputFilterWin64" Value="*.pgd;*.pgc"/>
+	<Property Name="PGOProfileCommandLineWin64" Value=""/>
+	<Property Name="PGOCompileArgsWin64" Value=""/>
+	<Property Name="PGOCompileExtraBinariesWin64" Value="*.pgd"/>
+
+	<!-- enable /FASTGENPROFILE when running locally -->
+	<Property Name="PGOCompileArgsWin64" Value="-PGOFastGen" If="!$(IsBuildMachine)"/>
+
+	<!-- no special preprocessing required for this platform -->
+	<Macro Name="PGOPrepareBuildForProfilingWin64" Arguments="LocalStagingDir;Configuration"/>
+
+
+</BuildGraph>
\ No newline at end of file
diff --git a/Engine/Build/Graph/Tasks/PGOProfileProject.xml b/Engine/Build/Graph/Tasks/PGOProfileProject.xml
new file mode 100644
--- /dev/null
+++ b/Engine/Build/Graph/Tasks/PGOProfileProject.xml
@@ -0,0 +1,94 @@
+<?xml version='1.0' ?>
+<BuildGraph xmlns="http://www.epicgames.com/BuildGraph" xmlns:xsi="http://www.w3.org/2001/XMLSchema-instance" xsi:schemaLocation="http://www.epicgames.com/BuildGraph ../../Schema.xsd" >
+
+	<!-- NOTE: this file is expecting many of the properties from BuildAndTestProject.xml or equivalent to be defined. -->
+	
+
+	<!-- Common PGO options -->
+	<Option Name="PGOAutoSubmitResults" DefaultValue="false" Restrict="true|false" Description="Auto Submit PGO Profile files"/>
+	<Option Name="PGOProfileTimeout" DefaultValue="7200" Description="Maximum length of time to allow PGO profile build to run"/>
+	
+
+	<!-- Enumerate all platforms & platform extensions -->
+	<Property Name="AllPGOPlatforms" Value=""/>
+	<Include Script="Inc/*/PGOProperties.xml"/>
+	<Include Script="../../../Platforms/*/Build/Inc/PGOProperties.xml"/>
+	
+	
+	<!-- Helper macro to create a "PGO Profile Replay <Platform>" node for a given platform -->
+	<!--     Platform:           the platform to create the node for -->
+	<!--     Configuration:      build configuration to use. Shipping or Test only -->
+	<!--     LocalReplay:        full path to the .replay file to use for PGO profiling -->
+	<!--     LocalStagingDir:    working directory where the PGO Profiling build is prepared, including the stage platform name -->
+	<!--     Build:              source path for a staged build, including the stage platform name. this is copied to the LocalStagingDir. Some platforms also require it to have a Shipping/Test configuration included.  -->
+	<!--     BuildRequires:      optional prerequisites for Build parameter (leave blank for none)-->
+	<!--     AgentOverride:      optional custom agent name for the node (leave blank for default) -->
+	<!--     CompileArgs:        optional additional compile arguments for PGO Profiling build -->
+	<!--     LaunchCommandLine:  optional additional command line to use when launching the PGO Profiling build  -->
+	<Macro Name="BasicReplayPGOProfile" Arguments="Platform;Configuration;LocalReplay;LocalStagingDir;Build" OptionalArguments="BuildRequires;AgentOverride;CompileArgs;LaunchCommandLine">
+
+		<!-- Select Test or Shipping configuration, favoring Shipping. In theory Debug and Development would work, but it's highly unlikely to be useful & this would suggest TargetConfigurations has been left at the default 'Development' -->
+		<Property Name="PGOConfiguration" Value=""/>
+		<Property Name="PGOConfiguration" Value="Test" If="ContainsItem('$(Configuration)','Test','+')"/>
+		<Property Name="PGOConfiguration" Value="Shipping" If="ContainsItem('$(Configuration)','Shipping','+')"/>
+		<Error Message="PGO only supports Test or Shipping" If="'$(PGOConfiguration)'==''"/>
+		
+		<!-- Determine host agent type -->
+		<Property Name="PGOHostAgentType" Value="Win64" />
+		<Property Name="PGOHostAgentType" Value="Mac" If="'$(Platform)' == 'Mac' or '$(Platform)' == 'IOS' or '$(Platform)' == 'tvOS'"/>
+		<Property Name="PGOHostAgentType" Value="Linux" If="'$(Platform)' == 'Linux'"/>
+		
+		<!-- Compiling will require this - built via BuildAndTestProject.xml for example -->
+		<Property Name="PGOPlatformToolsNodeName" Value="Compile Tools $(PGOHostAgentType)" />
+		
+		<!-- Set reasonable default agent type -->
+		<Property Name="PGOAgentType" Value="$(PGOHostAgentType)"/>
+		<Property Name="PGOAgentType" Value="$(AgentOverride)" If="'$(AgentOverride)' != ''"/>
+		
+		
+		<Agent Name="PGO Profile Replay Agent $(Platform)" Type="$(PGOAgentType)">
+
+			<!-- copy the given source build (cooked data etc) to a cleaned local staging directory -->
+			<Node Name="Copy $(Platform) Build" Produces="#$(Platform) PGO Build" Requires="$(BuildRequires)">
+				<Log Message="Cleaning local staging directory..."/>
+				<Delete Files="$(LocalStagingDir)/..."/>
+
+				<Log Message="Copying $(Build)/ to local staging directory..."/>
+				<Copy From="$(Build)/..." To="$(LocalStagingDir)/..."/>
+			</Node>
+
+			<!-- build a PGO profiling executable and copy the results over to the local staging directory -->
+			<Node Name="Compile PGO Profile $(Platform)" Requires="$(PGOPlatformToolsNodeName);#$(Platform) PGO Build" Produces="#$(Platform) PGO Profile Binaries">
+				<Compile Target="$(TargetName)" Platform="$(Platform)" Configuration="$(PGOConfiguration)" Arguments="-skipdeploy -PGOProfile $(CompileArgs) $(PGOCompileArgs$(Platform))" Tag="#$(Platform) PGOProfile Compile Results"/>
+
+				<Log Message="Copying binaries to local staging directory..."/> 
+				<Copy From="$(ProjectPath)/Binaries/$(Platform)/..." To="$(LocalStagingDir)/$(PGOBinaryStagingPath$(Platform))..."/>
+
+				<!-- the compile step may also generate additional files that will be used for PGOOptimize, such as the .pgd file on MSVC. Copy this to our output directory -->
+				<Log Message="Copying PGO support files to output directory, if any..."/>
+				<Copy Files="$(PGOCompileExtraBinaries$(Platform))" From="$(ProjectPath)/Binaries/$(Platform)/..." To="$(PGOProfileOutput$(Platform))/..."/>
+			</Node>
+			
+			<!-- run the basic replay PGO profile build to generate PGO profiling data -->
+			<Node Name="Run PGO Profile $(Platform)" Requires="#$(Platform) PGO Profile Binaries">
+				<Expand Name="PGOPrepareBuildForProfiling$(Platform)" LocalStagingDir="$(LocalStagingDir)" Configuration="$(Configuration)"/>
+				<Command Name="RunUnreal" Arguments="-test=BasicReplayPGONode -project='$(ProjectPath)/$(ProjectName).uproject' -platform=$(Platform) -configuration=$(PGOConfiguration) -build=$(LocalStagingDir) -LocalReplay='$(LocalReplay)' -ProfileOutputDirectory='$(PGOProfileOutput$(Platform))' -PgcFilenamePrefix=$(TargetName) -MaxDuration=$(PGOProfileTimeout) -AdditionalCommandLine='$(LaunchCommandLine) $(PGOProfileCommandLine$(Platform))'"/>
+			</Node>
+
+			<!-- submit the PGO profiling data -->
+			<Node Name="Gather PGO Profile Results $(Platform)" Requires="Run PGO Profile $(Platform)" Produces="#$(Platform) PGO Profile Complete">
+				<Tag Files="$(PGOProfileOutput$(Platform))/..." Filter="$(PGOProfileOutputFilter$(Platform))" With="#$(Platform) PGO Profile Results"/>
+				<Do If="$(PGOAutoSubmitResults)">
+					<Submit Files="#$(Platform) PGO Profile Results" FileType="+w" Description="Automated update of $(Platform) PGO files"/>
+				</Do>
+			</Node>
+			
+			<!-- all done -->
+			<Node Name ="PGO Profile Replay $(Platform)" Requires="Gather PGO Profile Results $(Platform)">
+				<Log Message="PGO Profiling completed for $(Platform)"/>
+			</Node>
+		</Agent>
+	</Macro>
+
+	
+</BuildGraph>
\ No newline at end of file
