Add -ThinLTO to the Android PGO command line for consistency with other Clang platforms.

 (note: Android PGO example not up and running yet)

#rnx

diff --git a/Engine/Build/Graph/Tasks/Inc/Android/PGOProperties.xml b/Engine/Build/Graph/Tasks/Inc/Android/PGOProperties.xml
--- a/Engine/Build/Graph/Tasks/Inc/Android/PGOProperties.xml
+++ b/Engine/Build/Graph/Tasks/Inc/Android/PGOProperties.xml
@@ -11,7 +11,7 @@
 	<Property Name="PGOProfileOutputPathAndroid" Value="$(ProjectPath)/Platforms/Android/Build/PGO"/>
 	<Property Name="PGOProfileOutputFilterAndroid" Value="*.profdata"/>
 	<Property Name="PGOProfileCommandLineAndroid" Value=""/>
-	<Property Name="PGOCompileArgsAndroid" Value=""/>
+	<Property Name="PGOCompileArgsAndroid" Value="-ThinLTO"/>
 	<Property Name="PGOCompileExtraBinariesAndroid" Value=""/>
 
 	<!-- no special preprocessing required for this platform -->
