Updating sample build scripts to not build UnrealHeaderTool

diff --git a/Engine/Build/CompileShadersTest.xml b/Engine/Build/CompileShadersTest.xml
--- a/Engine/Build/CompileShadersTest.xml
+++ b/Engine/Build/CompileShadersTest.xml
@@ -44,8 +44,6 @@
 		<Property Name="CompileNodeName" Value="Compile $(Project) $(EditorPlatform)"/>
 		<Agent Name="Compile Editor Agent $(EditorPlatform) $(Project)" Type="$(EditorPlatform)">
 			<Node Name="$(CompileNodeName)">
-				<Compile Target="UnrealHeaderTool" Platform="$(EditorPlatform)" Configuration="Development"/>
-				<Compile Target="UnrealHeaderTool" Platform="$(EditorPlatform)" Project="$(Project)" Configuration="Development"/>
 				<Compile Target="ShaderCompileWorker" Platform="$(EditorPlatform)" Configuration="Development"/>
 				<Compile Target="CrashReportClientEditor" Platform="$(EditorPlatform)" Project="$(Project)" Configuration="Shipping"/>
 				<Compile Target="$(EditorTarget)" Platform="$(EditorPlatform)" Project="$(Project)" Configuration="Development"/>
diff --git a/Engine/Build/Graph/Examples/BuildEditorAndTools.xml b/Engine/Build/Graph/Examples/BuildEditorAndTools.xml
--- a/Engine/Build/Graph/Examples/BuildEditorAndTools.xml
+++ b/Engine/Build/Graph/Examples/BuildEditorAndTools.xml
@@ -62,7 +62,6 @@
 		
 		<!-- Compile the tool executables -->
 		<Node Name="Compile Tools Win64" Requires="Update Version Files" Produces="#ToolBinaries">
-			<Compile Target="UnrealHeaderTool" Platform="Win64" Configuration="Development" Tag="#ToolBinaries"/>
 			<Compile Target="ShaderCompileWorker" Platform="Win64" Configuration="Development" Tag="#ToolBinaries"/>
 			<Compile Target="UnrealLightmass" Platform="Win64" Configuration="Development" Tag="#ToolBinaries"/>
 			<Compile Target="InterchangeWorker" Platform="Win64" Configuration="Development" Tag="#ToolBinaries"/>
diff --git a/Engine/Build/Graph/Examples/BuildWorldPartitionHLODs.xml b/Engine/Build/Graph/Examples/BuildWorldPartitionHLODs.xml
--- a/Engine/Build/Graph/Examples/BuildWorldPartitionHLODs.xml
+++ b/Engine/Build/Graph/Examples/BuildWorldPartitionHLODs.xml
@@ -41,7 +41,6 @@
 		</Node>
 		
 		<Node Name="Compile Tools Win64" Requires="Update Version Files">
-			<Compile Target="UnrealHeaderTool" Platform="Win64" Configuration="Development"/>
 			<Compile Target="ShaderCompileWorker" Platform="Win64" Configuration="Development"/>
 		</Node>
 		
diff --git a/Engine/Build/Graph/Examples/Building.xml b/Engine/Build/Graph/Examples/Building.xml
--- a/Engine/Build/Graph/Examples/Building.xml
+++ b/Engine/Build/Graph/Examples/Building.xml
@@ -3,13 +3,13 @@
 
 	<Agent Name="Properties Agent" Type="CompileWin64">
 
-		<!-- Compile UnrealHeaderTool -->
-		<Node Name="Compile UnrealHeaderTool">
-			<Compile Target="UnrealHeaderTool" Platform="Win64" Configuration="Development" Tag="#UHTFiles"/>
-			<Log Message="The following files are part of UnrealHeaderTool:" Files="#UHTFiles"/>
+		<!-- Compile UnrealPak -->
+		<Node Name="Compile UnrealPak">
+			<Compile Target="UnrealPak" Platform="Win64" Configuration="Development" Tag="#UnrealPakFiles"/>
+			<Log Message="The following files are part of UnrealPak:" Files="#UnrealPakFiles"/>
 		</Node>
 		
 	</Agent>
 	
-	<Aggregate Name="Building Examples" Requires="Compile UnrealHeaderTool"/>
+	<Aggregate Name="Building Examples" Requires="Compile UnrealPak"/>
 </BuildGraph>
diff --git a/Engine/Build/Graph/Tasks/BuildAndTestProject.xml b/Engine/Build/Graph/Tasks/BuildAndTestProject.xml
--- a/Engine/Build/Graph/Tasks/BuildAndTestProject.xml
+++ b/Engine/Build/Graph/Tasks/BuildAndTestProject.xml
@@ -171,7 +171,6 @@
 			</Node>
 			<!-- First generate nodes to build the required tools (e.g. UHT, ShaderCompiler, Editor -->
 			<Node Name="Compile Tools $(EditorPlatform)" Requires="Update Version Files $(EditorPlatform)">
-				<Compile Target="UnrealHeaderTool" Platform="$(EditorPlatform)" Configuration="Development" Arguments="$(GenericCompileArguments)"/>
 				<Compile Target="ShaderCompileWorker" Platform="$(EditorPlatform)" Configuration="Development" Arguments="$(GenericCompileArguments)"/>
 				<Compile Target="UnrealPak" Platform="$(EditorPlatform)" Configuration="Development" Arguments="$(GenericCompileArguments)"/>
 				<Compile Target="CrashReportClient" Platform="$(EditorPlatform)" Configuration="Shipping" Arguments="$(GenericCompileArguments)"/>
diff --git a/Engine/Build/ShaderTypeRegressionTest.xml b/Engine/Build/ShaderTypeRegressionTest.xml
--- a/Engine/Build/ShaderTypeRegressionTest.xml
+++ b/Engine/Build/ShaderTypeRegressionTest.xml
@@ -43,8 +43,6 @@
 		<Property Name="CompileNodeName" Value="Compile $(Project) $(EditorPlatform)"/>
 		<Agent Name="Compile Editor Agent $(EditorPlatform) $(Project)" Type="$(EditorPlatform)">
 			<Node Name="$(CompileNodeName)">
-				<Compile Target="UnrealHeaderTool" Platform="$(EditorPlatform)" Configuration="Development"/>
-				<Compile Target="UnrealHeaderTool" Platform="$(EditorPlatform)" Project="$(Project)" Configuration="Development"/>
 				<Compile Target="ShaderCompileWorker" Platform="$(EditorPlatform)" Configuration="Development"/>
 				<Compile Target="CrashReportClientEditor" Platform="$(EditorPlatform)" Project="$(Project)" Configuration="Shipping"/>
 				<Compile Target="$(EditorTarget)" Platform="$(EditorPlatform)" Project="$(Project)" Configuration="Development"/>
