Make Project name optional to be prepended to node name and label name

diff --git a/Engine/Build/Graph/Tasks/BuildAndTestProject.xml b/Engine/Build/Graph/Tasks/BuildAndTestProject.xml
--- a/Engine/Build/Graph/Tasks/BuildAndTestProject.xml
+++ b/Engine/Build/Graph/Tasks/BuildAndTestProject.xml
@@ -12,6 +12,11 @@
 	<Option Name="EditorTargetName" DefaultValue="" Restrict="" Description="Specify the editor target name if different than (ProjectName)Editor"/>
 	<Option Name="CookedEditorTargetName" DefaultValue="" Restrict="" Description="Specify the cooked editor target name if different than (ProjectName)CookedEditor"/>
 
+	<Option Name="PreNodeName" DefaultValue="$(ProjectName)" Description="String to prepend to node name"/>
+	<Option Name="PreLabelName" DefaultValue="" Description="String to prepend to label name"/>
+	<Property Name="PreNodeName" Value="$(PreNodeName) " If="'$(PreNodeName)' != ''"/>
+	<Property Name="PreLabelName" Value="$(PreLabelName) " If="'$(PreLabelName)' != ''"/>
+
 	<Option Name="GenericCompileArguments" DefaultValue="" Restrict="" Description="Configurable compile args applied to all compile targets."/>
 	<Option Name="ExtraToolCompileArguments" DefaultValue="" Restrict="" Description="Extra arguments to use when building the tools" />
 	<Option Name="ExtraEditorCompileArguments" DefaultValue="" Restrict="" Description="Extra arguments to use when building the editor" />
@@ -138,10 +143,6 @@
 	<Property Name="TargetType" Value="Game" />
 	<Property Name="TargetType" Value="Client" If="Contains('$(TargetName)', 'Client')" />
 
-	<!-- We want a display name for a number of jobs which is best if it's the target (e.g. ShooterClient), but not if it's content Only -->
-	<Property Name="DisplayTargetName" Value="$(TargetName)" />
-	<Property Name="DisplayTargetName" Value="$(ProjectName)" If="$(TargetName) == 'UnrealGame'" />
-	
 	<!-- These properties will be filled in with nodes that are generated based on options -->
 	<Property Name="GeneratedToolNodes" Value="" />
 	<Property Name="GeneratedEditorTestNodes" Value="" />
@@ -166,11 +167,11 @@
 
 		<Agent Name="Build $(ProjectName)Editor $(EditorPlatform)" Type="$(HostAgentType)">
 			<!-- Update the Version.h and Build.version files if we're making a formal build on a build machine -->
-			<Node Name="Update Version Files $(EditorPlatform)">
+			<Node Name="$(PreNodeName)Update Version Files $(EditorPlatform)">
 				<SetVersion Change="$(Change)" Branch="$(EscapedBranch)" If="$(Versioned) and $(IsBuildMachine)"/>
 			</Node>
 			<!-- First generate nodes to build the required tools (e.g. UHT, ShaderCompiler, Editor -->
-			<Node Name="Compile Tools $(EditorPlatform)" Requires="Update Version Files $(EditorPlatform)">
+			<Node Name="$(PreNodeName)Compile Tools $(EditorPlatform)" Requires="$(PreNodeName)Update Version Files $(EditorPlatform)">
 				<Compile Target="ShaderCompileWorker" Platform="$(EditorPlatform)" Configuration="Development" Arguments="$(GenericCompileArguments)"/>
 				<Compile Target="UnrealPak" Platform="$(EditorPlatform)" Configuration="Development" Arguments="$(GenericCompileArguments)"/>
 				<Compile Target="CrashReportClient" Platform="$(EditorPlatform)" Configuration="Shipping" Arguments="$(GenericCompileArguments)"/>
@@ -185,25 +186,25 @@
 				</Do>
 				<Compile Target="UnrealInsights" Platform="$(EditorPlatform)" Configuration="Shipping" Arguments="$(GenericCompileArguments)" If="$(IncludeUnrealInsightsInToolCompile)"/>
 			</Node>
-			<Node Name="Compile $(ProjectName)Editor $(EditorPlatform)" Requires="Compile Tools $(EditorPlatform)">
+			<Node Name="$(PreNodeName)Compile Editor $(EditorPlatform)" Requires="$(PreNodeName)Compile Tools $(EditorPlatform)">
 				<Compile Target="$(EditorTargetName)" Platform="$(EditorPlatform)" Configuration="Development" Arguments="$(GenericCompileArguments) $(ExtraEditorCompileArguments)"/>
 			</Node>
 			<!-- Cooked editor target -->
-			<Node Name="Compile $(ProjectName)CookedEditor $(EditorPlatform)" Requires="Compile Tools $(EditorPlatform)">
+			<Node Name="$(PreNodeName)Compile CookedEditor $(EditorPlatform)" Requires="$(PreNodeName)Compile Tools $(EditorPlatform)">
 				<ForEach Name="Config" Values="$(CookedEditorConfigurations)" Separator="+">
 					<Compile Target="$(CookedEditorTargetName)" Platform="$(EditorPlatform)" Configuration="$(Config)" Arguments="$(GenericCompileArguments) $(ExtraEditorCompileArguments)"/>
 				</ForEach>
 			</Node>
 			<!--Add these nodes to our dependency list -->
-			<Property Name="GeneratedToolNodes" Value="Compile $(ProjectName)Editor $(EditorPlatform); Compile Tools $(EditorPlatform)" If="!$(SkipBuild)" />
+			<Property Name="GeneratedToolNodes" Value="$(PreNodeName)Compile Editor $(EditorPlatform);$(PreNodeName)Compile Tools $(EditorPlatform)" If="!$(SkipBuild)" />
 			<!-- Editor Labels -->
-			<Label Category="Editors" Name="$(ProjectName) Editor $(EditorPlatform)" Requires="$(GeneratedToolNodes)" />
+			<Label Category="Editors" Name="$(PreLabelName)Editor $(EditorPlatform)" Requires="$(GeneratedToolNodes)" />
 		</Agent>
 		<Agent Name="Test $(ProjectName)Editor $(EditorPlatform)" Type="$(TestAgentType)">
 			<!-- Prerequisites for the editor test to run -->
 			<Property Name="TestPrerequisites" Value="$(GeneratedToolNodes)"/>
 			<!-- Declare a boot test for the editor-->
-			<Property Name="PlatformBootTestNodeName" Value="Editor $(EditorPlatform) Test=BootTest" />
+			<Property Name="PlatformBootTestNodeName" Value="$(PreNodeName)Editor $(EditorPlatform) Test=BootTest" />
 			<Do If="'$(AutomationReportOutputDirectory)' != ''">
 				<Property Name="OutputFragmentPath" Value="$(EditorPlatform)Editor/UE.EditorBootTest" />
 				<Property Name="OutputReportPath" Value="$(AutomationReportOutputDirectory)/$(OutputFragmentPath)" />
@@ -213,7 +214,7 @@
 			<Node Name="$(PlatformBootTestNodeName)" Requires="$(GeneratedToolNodes)">
 				<Command Name="RunUnreal" Arguments="-project=$(TargetProject) -build=Editor -test=&quot;UE.EditorBootTest&quot; $(GauntletEditorTestArgs) $(GauntletReportArgs)" />
 			</Node>
-			<Property Name="ProjectEditorDDC" Value="Editor DDC $(EditorPlatform)" />
+			<Property Name="ProjectEditorDDC" Value="$(PreNodeName)Editor DDC $(EditorPlatform)" />
 			<Node Name="$(ProjectEditorDDC)" Requires="$(GeneratedToolNodes)">
 				<Property Name="DDCTargetPlatform" Value="$(EditorPlatform)" />
 				<Property Name="DDCTargetPlatform" Value="Windows" If="'$(DDCTargetPlatform)' == 'Win64'"/>
@@ -231,7 +232,7 @@
 					<Property Name="GauntletReportArgs" Value="-logdir=&quot;$(OutputReportPath)&quot; -artifactname=Gauntlet -ReportExportPath=&quot;$(OutputReportPath)&quot;" />
 					<Property Name="GauntletReportArgs" Value="$(GauntletReportArgs) -ReportURL=&quot;$(AutomationReportUrl)/$(OutputFragmentPath)/&quot;" If="'$(AutomationReportUrl)' != ''" />
 				</Do>
-				<Property Name="TestNodeName" Value="Editor $(EditorPlatform) Test=$(TestName)" />
+				<Property Name="TestNodeName" Value="$(PreNodeName)Editor $(EditorPlatform) Test=$(TestName)" />
 				<Node Name="$(TestNodeName)" Requires="$(TestPrerequisites)">
 					<Command Name="RunUnreal" Arguments="-project=$(TargetProject) -build=Editor -test=&quot;$(TestName)&quot; $(GauntletReportArgs) $(GauntletEditorTestArgs)" />
 				</Node>
@@ -240,12 +241,12 @@
 			</ForEach>
 
 			<!-- Create Labels -->
-			<Label Category="Editors" Name="$(ProjectName) Editor DDC $(EditorPlatform)" Requires="$(ProjectEditorDDC)" />
+			<Label Category="Editors" Name="$(PreLabelName)Editor DDC $(EditorPlatform)" Requires="$(ProjectEditorDDC)" />
 
 			<!-- If this editor was requested (as opposed to being required due to target platforms) then link in the test nodes and set up labels -->
 			<Do If="ContainsItem('$(EditorPlatforms)', '$(EditorPlatform)', '+') and !ContainsItem('$(EditorPlatformsToSkipTests)', '$(EditorPlatform)', '+')">
-				<Label Category="Tests" Name="BootTest $(ProjectName)Editor $(EditorPlatform)" Requires="$(PlatformBootTestNodeName)" />
-				<Label Category="Tests" Name="Test $(ProjectName)Editor $(EditorPlatform)" Requires="$(PlatformEditorTestNodes)" />
+				<Label Category="Tests" Name="$(PreLabelName)BootTest Editor $(EditorPlatform)" Requires="$(PlatformBootTestNodeName)" />
+				<Label Category="Tests" Name="$(PreLabelName)Test Editor $(EditorPlatform)" Requires="$(PlatformEditorTestNodes)" />
 				<Property Name="GeneratedEditorTestNodes" Value="$(GeneratedEditorTestNodes);$(PlatformBootTestNodeName)" If="$(EditorBootTest)" />
 				<Property Name="GeneratedEditorTestNodes" Value="$(GeneratedEditorTestNodes);$(ProjectEditorDDC)" If="$(FillEditorDDC)" />
 				<Property Name="GeneratedEditorTestNodes" Value="$(GeneratedEditorTestNodes);$(PlatformEditorTestNodes)" />
@@ -262,14 +263,14 @@
 		<Property Name="EditorPlatformPath" Value="$(EditorPlatform)" />
 		<Property Name="EditorPlatformPath" Value="Windows" If="'$(EditorPlatform)' == 'Win64'" />
 
-		<Property Name="PlatformPublishNodeName" Value="Publish CookedEditor $(EditorPlatform)" />
+		<Property Name="PlatformPublishNodeName" Value="$(PreNodeName)Publish CookedEditor $(EditorPlatform)" />
 
-		<Property Name="CookedEditorCookRequirements" Value="Compile Tools $(EditorPlatform);Compile $(ProjectName)Editor $(EditorPlatform);Compile $(ProjectName)CookedEditor $(EditorPlatform)"/>
+		<Property Name="CookedEditorCookRequirements" Value="$(PreNodeName)Compile Tools $(EditorPlatform);$(PreNodeName)Compile Editor $(EditorPlatform);$(PreNodeName)Compile CookedEditor $(EditorPlatform)"/>
 		<Property Name="CookedEditorCookerArgs" Value="-project=$(TargetProject) -platform=$(EditorPlatform) -compressed $(ExtraCookArguments)"/>
 
 		<Agent Name="CookedEditor Cook and Publish $(EditorPlatform)" Type="$(HostAgentType)">
 			<!-- Cook Editor -->
-			<Node Name="Cook CookedEditor $(EditorPlatform)" Requires="$(CookedEditorCookRequirements)">
+			<Node Name="$(PreNodeName)Cook CookedEditor $(EditorPlatform)" Requires="$(CookedEditorCookRequirements)">
 				<Command Name="MakeCookedEditor" Arguments="-cook $(CookedEditorCookerArgs)"/>
 			</Node>
 
@@ -278,12 +279,12 @@
 			<Property Name="CookedEditorStageRequirements" Value="$(CookedEditorCookRequirements)"/>
 			<Property Name="CookedEditorStageRequirements" Value="$(CookedEditorStageRequirements);Cook CookedEditor $(EditorPlatform)" If="!$(SkipCook)"/>
 
-			<Node Name="Stage CookedEditor $(EditorPlatform)" Requires="$(CookedEditorStageRequirements)">
+			<Node Name="$(PreNodeName)Stage CookedEditor $(EditorPlatform)" Requires="$(CookedEditorStageRequirements)">
 				<Command Name="MakeCookedEditor" Arguments="-project=$(TargetProject) $(CookedEditorStageArgs)"/>
 			</Node>
 
 			<!-- Publish CookedEditor -->
-			<Node Name="Publish CookedEditor $(EditorPlatform)" Requires="Stage CookedEditor $(EditorPlatform)">
+			<Node Name="$(PreNodeName)Publish CookedEditor $(EditorPlatform)" Requires="Stage CookedEditor $(EditorPlatform)">
 				<Log Message="Copying to $(NetworkOutputDirectory)/$(EditorPlatformPath)CookedEditor"/>
 				<Do If="!$(SkipPublish)">
 					<Copy Files="..." From="$(ProjectOutputDirectory)/$(EditorPlatformPath)CookedEditor" To="$(NetworkOutputDirectory)/$(EditorPlatformPath)CookedEditor"/>
@@ -291,8 +292,8 @@
 			</Node>
 
 			<!-- Aggregate Nodes requirements -->
-			<Property Name="GeneratedCookedEditorNodes" Value="Compile $(ProjectName)CookedEditor $(EditorPlatform)" If="!$(SkipBuild)" />
-			<Property Name="GeneratedCookedEditorNodes" Value="$(GeneratedCookedEditorNodes);Cook CookedEditor $(EditorPlatform)" If="!$(SkipCook)" />
+			<Property Name="GeneratedCookedEditorNodes" Value="$(PreNodeName)Compile CookedEditor $(EditorPlatform)" If="!$(SkipBuild)" />
+			<Property Name="GeneratedCookedEditorNodes" Value="$(GeneratedCookedEditorNodes);$(PreNodeName)Cook CookedEditor $(EditorPlatform)" If="!$(SkipCook)" />
 			<Property Name="GeneratedCookedEditorNodes" Value="$(GeneratedCookedEditorNodes);$(PlatformPublishNodeName)" If="!$(SkipPublish) And $(PublishBuilds)" />
 		</Agent>
 
@@ -303,7 +304,7 @@
 			<Property Name="TestPrerequisites" Value="$(GeneratedCookedEditorNodes)"/>
 			<Property Name="TestPrerequisites" Value="$(TestPrerequisites);$(PlatformPublishNodeName)" If="!$(PublishBuilds)"/>
 			<!-- Declare a boot test for the editor-->
-			<Property Name="PlatformBootTestNodeName" Value="CookedEditor $(EditorPlatform) Test=BootTest" />
+			<Property Name="PlatformBootTestNodeName" Value="$(PreNodeName)CookedEditor $(EditorPlatform) Test=BootTest" />
 			<Do If="'$(AutomationReportOutputDirectory)' != ''">
 				<Property Name="OutputFragmentPath" Value="$(EditorPlatform)Editor/UE.CookedEditorBootTest" />
 				<Property Name="OutputReportPath" Value="$(AutomationReportOutputDirectory)/$(OutputFragmentPath)" />
@@ -327,7 +328,7 @@
 					<Property Name="GauntletReportArgs" Value="-logdir=&quot;$(OutputReportPath)&quot; -artifactname=Gauntlet -ReportExportPath=&quot;$(OutputReportPath)&quot;" />
 					<Property Name="GauntletReportArgs" Value="$(GauntletReportArgs) -ReportURL=&quot;$(AutomationReportUrl)/$(OutputFragmentPath)/&quot;" If="'$(AutomationReportUrl)' != ''" />
 				</Do>
-				<Property Name="TestNodeName" Value="CookedEditor $(EditorPlatform) Test=$(TestName)" />
+				<Property Name="TestNodeName" Value="$(PreNodeName)CookedEditor $(EditorPlatform) Test=$(TestName)" />
 				<Node Name="$(TestNodeName)" Requires="$(TestPrerequisites)">
 					<Command Name="RunUnreal" Arguments="-project=$(TargetProject) -build=&quot;$(NetworkOutputDirectory)/$(EditorPlatformPath)CookedEditor&quot; -test=&quot;$(TestName)&quot; $(GauntletReportArgs) $(CustomModuleRoleArg) $(GauntletCookedEditorTestArgs)" />
 				</Node>
@@ -338,15 +339,15 @@
 			<!-- Create Labels -->
 			<!-- If this cooked editor was requested (as opposed to being required due to target platforms) then link in the test nodes and set up labels -->
 			<Do If="ContainsItem('$(CookedEditorPlatforms)', '$(EditorPlatform)', '+') and !ContainsItem('$(CookedEditorPlatformsToSkipTests)', '$(EditorPlatform)', '+')">
-				<Label Category="Tests" Name="BootTest $(ProjectName)CookedEditor $(EditorPlatform)" Requires="$(PlatformBootTestNodeName)" />
-				<Label Category="Tests" Name="Test $(ProjectName)CookedEditor $(EditorPlatform)" Requires="$(PlatformCookedEditorTestNodes)" />
+				<Label Category="Tests" Name="$(PreLabelName)BootTest CookedEditor $(EditorPlatform)" Requires="$(PlatformBootTestNodeName)" />
+				<Label Category="Tests" Name="$(PreLabelName)Test CookedEditor $(EditorPlatform)" Requires="$(PlatformCookedEditorTestNodes)" />
 				<Property Name="GeneratedCookedEditorTestNodes" Value="$(GeneratedCookedEditorTestNodes);$(PlatformBootTestNodeName)" If="$(CookedEditorBootTest)" />
 				<Property Name="GeneratedCookedEditorTestNodes" Value="$(GeneratedCookedEditorTestNodes);$(PlatformCookedEditorTestNodes)" />
 			</Do>
 		</Agent>
 
 		<!-- Add to Editor Labels -->
-		<Label Category="Editors" Name="$(ProjectName) CookedEditor $(EditorPlatform)" Requires="$(GeneratedCookedEditorNodes)" />
+		<Label Category="Editors" Name="$(PreLabelName)CookedEditor $(EditorPlatform)" Requires="$(GeneratedCookedEditorNodes)" />
 	</ForEach>
 
 	<!-- Now, for each platform we build, cook, stage, package, publish, and test -->
@@ -374,19 +375,19 @@
 		<Property Name="TestAgentType" Value="$(TestAgentOverride)" If="'$(TestAgentOverride)' != ''"/>
 
 		<!-- Building and cooking will require these nodes from the appropriate host. They should have been built above... -->
-		<Property Name="PlatformToolsNodeName" Value="Compile Tools $(HostAgentType)" />
-		<Property Name="PlatformEditorNodeName" Value="Compile $(ProjectName)Editor $(HostAgentType)" />
+		<Property Name="PlatformToolsNodeName" Value="$(PreNodeName)Compile Tools $(HostAgentType)" />
+		<Property Name="PlatformEditorNodeName" Value="$(PreNodeName)Compile Editor $(HostAgentType)" />
 		<!-- Define these names up front as we use some of them several times -->
-		<Property Name="PlatformAgentName" Value="$(DisplayTargetName) $(TargetPlatform)" />
-		<Property Name="PlatformCompileNodeName" Value="Compile $(DisplayTargetName) $(TargetPlatform)" />
+		<Property Name="PlatformAgentName" Value="$(ProjectName) $(TargetPlatform)" />
+		<Property Name="PlatformCompileNodeName" Value="$(PreNodeName)Compile $(TargetPlatform)" />
 		<Property Name="UploadSymbolsNodeName" Value="" />
-		<Property Name="UploadSymbolsNodeName" Value="Upload Symbols $(TargetPlatform)" If="'$(SymbolServerPath)' != ''" />
-		<Property Name="PlatformCookNodeName" Value="Cook $(TargetPlatform)" />
-		<Property Name="PlatformStageNodeName" Value="Stage $(TargetPlatform)" />
-		<Property Name="PlatformPackageNodeName" Value="Package $(TargetPlatform)" />
-		<Property Name="PlatformHashNodeName" Value="Hash $(TargetPlatform)" />
-		<Property Name="PlatformPublishStagedNodeName" Value="Publish Staged $(TargetPlatform)" />
-		<Property Name="PlatformPublishPackagedNodeName" Value="Publish Packaged $(TargetPlatform)" />
+		<Property Name="UploadSymbolsNodeName" Value="$(PreNodeName)Upload Symbols $(TargetPlatform)" If="'$(SymbolServerPath)' != ''" />
+		<Property Name="PlatformCookNodeName" Value="$(PreNodeName)Cook $(TargetPlatform)" />
+		<Property Name="PlatformStageNodeName" Value="$(PreNodeName)Stage $(TargetPlatform)" />
+		<Property Name="PlatformPackageNodeName" Value="$(PreNodeName)Package $(TargetPlatform)" />
+		<Property Name="PlatformHashNodeName" Value="$(PreNodeName)Hash $(TargetPlatform)" />
+		<Property Name="PlatformPublishStagedNodeName" Value="$(PreNodeName)Publish Staged $(TargetPlatform)" />
+		<Property Name="PlatformPublishPackagedNodeName" Value="$(PreNodeName)Publish Packaged $(TargetPlatform)" />
 		
 		<!-- Set CookPlatform since it isn't always the same as the target -->
 		<Property Name="CookPlatform" Value="$(TargetPlatform)" />
@@ -522,7 +523,7 @@
 				<Property Name="GauntletTestArgs" Value="$(GauntletTargetTestAltArgs)" If="ContainsItem('$(AlternateDevicePoolPlatforms)', '$(TargetPlatform)', ';')"/>
 
 				<!-- Declare a boot test for the target-->	
-				<Property Name="BootTestName" Value="$(TargetPlatform) $(TestConfiguration) BootTest" />
+				<Property Name="BootTestName" Value="$(PreNodeName)$(TargetPlatform) $(TestConfiguration) BootTest" />
 				<Node Name="$(BootTestName)" After="$(PlatformPublishPackagedNodeName)" Requires="$(TestPrerequisites)">
 					<Do If="'$(AutomationReportOutputDirectory)' != ''">
 						<Property Name="OutputFragmentPath" Value="$(TargetPlatform)/$(TestConfiguration)/UE.BootTest" />
@@ -548,7 +549,7 @@
 						<Property Name="GauntletReportArgs" Value="-logdir=&quot;$(OutputReportPath)&quot; -artifactname=Gauntlet -ReportExportPath=&quot;$(OutputReportPath)&quot;" />
 						<Property Name="GauntletReportArgs" Value="$(GauntletReportArgs) -ReportURL=&quot;$(AutomationReportUrl)/$(OutputFragmentPath)/&quot;" If="'$(AutomationReportUrl)' != ''" />
 					</Do>
-					<Property Name="TestNodeName" Value="$(TargetPlatform) $(TestConfiguration) $(TestName)" />
+					<Property Name="TestNodeName" Value="$(PreNodeName)$(TargetPlatform) $(TestConfiguration) $(TestName)" />
 					<!-- Run the test. No 'Requires' as we run on the same machine locally and use the publish folder otherwise -->
 					<Node Name="$(TestNodeName)" Requires="$(TestPrerequisites)">
 						<Command Name="RunUnreal" Arguments="-project=$(TargetProject) -test=&quot;$(TestName)&quot; -platform=$(TargetPlatform) -configuration=$(TestConfiguration) $(GauntletBuildArg) $(GauntletTestArgs) $(GauntletReportArgs)"/>
@@ -562,11 +563,11 @@
 		</Agent>
 
 		<!-- Declare labels for CIS -->
-		<Label Category="Clients" Name="$(ProjectName) $(TargetPlatform)" Requires="$(PlatformCompileNodeName);$(UploadSymbolsNodeName)" />
-		<Label Category="Clients" Name="Content $(DisplayTargetName) $(TargetPlatform)" Requires="$(PlatformCookNodeName);$(PlatformStageNodeName);$(PlatformPackageNodeName)" />
-		<Label Category="Clients" Name="Publish $(DisplayTargetName) $(TargetPlatform)" Requires="$(PlatformPublishPackagedNodeName);$(PlatformPublishPackagedNodeName);$(PlatformPublishStagedNodeName);$(PlatformHashNodeName)" />
-		<Label Category="Tests" Name="BootTest $(DisplayTargetName) $(TargetPlatform)" Requires="$(AllPlatformBootTestNodes)" />
-		<Label Category="Tests" Name="Test $(DisplayTargetName) $(TargetPlatform)" Requires="$(AllPlatformTestNodes)" />
+		<Label Category="Clients" Name="$(PreLabelName)$(TargetPlatform)" Requires="$(PlatformCompileNodeName);$(UploadSymbolsNodeName)" />
+		<Label Category="Clients" Name="$(PreLabelName)Content $(TargetPlatform)" Requires="$(PlatformCookNodeName);$(PlatformStageNodeName);$(PlatformPackageNodeName)" />
+		<Label Category="Clients" Name="$(PreLabelName)Publish $(TargetPlatform)" Requires="$(PlatformPublishPackagedNodeName);$(PlatformPublishPackagedNodeName);$(PlatformPublishStagedNodeName);$(PlatformHashNodeName)" />
+		<Label Category="Tests" Name="$(PreLabelName)BootTest $(TargetPlatform)" Requires="$(AllPlatformBootTestNodes)" />
+		<Label Category="Tests" Name="$(PreLabelName)Test $(TargetPlatform)" Requires="$(AllPlatformTestNodes)" />
 
 		<!-- Tie all steps into our platform tasks unless excluded -->
 		<Property Name="PlatformTasks" Value="" />
diff --git a/Engine/Build/Graph/Tasks/PGOProfileProject.xml b/Engine/Build/Graph/Tasks/PGOProfileProject.xml
--- a/Engine/Build/Graph/Tasks/PGOProfileProject.xml
+++ b/Engine/Build/Graph/Tasks/PGOProfileProject.xml
@@ -39,7 +39,7 @@
 		<Property Name="PGOHostAgentType" Value="Linux" If="'$(Platform)' == 'Linux'"/>
 		
 		<!-- Compiling will require this - built via BuildAndTestProject.xml for example -->
-		<Property Name="PGOPlatformToolsNodeName" Value="Compile Tools $(PGOHostAgentType)" />
+		<Property Name="PGOPlatformToolsNodeName" Value="$(PreNodeName)Compile Tools $(PGOHostAgentType)" />
 		
 		<!-- Set reasonable default agent type -->
 		<Property Name="PGOAgentType" Value="$(PGOHostAgentType)"/>
@@ -49,7 +49,7 @@
 		<Agent Name="PGO Profile Replay Agent $(Platform)" Type="$(PGOAgentType)">
 
 			<!-- copy the given source build (cooked data etc) to a cleaned local staging directory -->
-			<Node Name="Copy $(Platform) Build" Produces="#$(Platform) PGO Build" Requires="$(BuildRequires)">
+			<Node Name="$(PreNodeName)Copy $(Platform) Build" Produces="#$(Platform) PGO Build" Requires="$(BuildRequires)">
 				<Log Message="Cleaning local staging directory..."/>
 				<Delete Files="$(LocalStagingDir)/..."/>
 
@@ -58,7 +58,7 @@
 			</Node>
 
 			<!-- build a PGO profiling executable and copy the results over to the local staging directory -->
-			<Node Name="Compile PGO Profile $(Platform)" Requires="$(PGOPlatformToolsNodeName);#$(Platform) PGO Build" Produces="#$(Platform) PGO Profile Binaries">
+			<Node Name="$(PreNodeName)Compile PGO Profile $(Platform)" Requires="$(PGOPlatformToolsNodeName);#$(Platform) PGO Build" Produces="#$(Platform) PGO Profile Binaries">
 				<Compile Target="$(TargetName)" Platform="$(Platform)" Configuration="$(PGOConfiguration)" Arguments="-skipdeploy -PGOProfile $(CompileArgs) $(PGOCompileArgs$(Platform))" Tag="#$(Platform) PGOProfile Compile Results"/>
 
 				<Log Message="Copying binaries to local staging directory..."/> 
@@ -70,13 +70,13 @@
 			</Node>
 			
 			<!-- run the basic replay PGO profile build to generate PGO profiling data -->
-			<Node Name="Run PGO Profile $(Platform)" Requires="#$(Platform) PGO Profile Binaries">
+			<Node Name="$(PreNodeName)Run PGO Profile $(Platform)" Requires="#$(Platform) PGO Profile Binaries">
 				<Expand Name="PGOPrepareBuildForProfiling$(Platform)" LocalStagingDir="$(LocalStagingDir)" Configuration="$(Configuration)"/>
 				<Command Name="RunUnreal" Arguments="-test=BasicReplayPGONode -project='$(ProjectPath)/$(ProjectName).uproject' -platform=$(Platform) -configuration=$(PGOConfiguration) -build=$(LocalStagingDir) -LocalReplay='$(LocalReplay)' -ProfileOutputDirectory='$(PGOProfileOutput$(Platform))' -PgcFilenamePrefix=$(TargetName) -MaxDuration=$(PGOProfileTimeout) -AdditionalCommandLine='$(LaunchCommandLine) $(PGOProfileCommandLine$(Platform))'"/>
 			</Node>
 
 			<!-- submit the PGO profiling data -->
-			<Node Name="Gather PGO Profile Results $(Platform)" Requires="Run PGO Profile $(Platform)" Produces="#$(Platform) PGO Profile Complete">
+			<Node Name="$(PreNodeName)Gather PGO Profile Results $(Platform)" Requires="$(PreNodeName)Run PGO Profile $(Platform)" Produces="#$(Platform) PGO Profile Complete">
 				<Tag Files="$(PGOProfileOutput$(Platform))/..." Filter="$(PGOProfileOutputFilter$(Platform))" With="#$(Platform) PGO Profile Results"/>
 				<Do If="$(PGOAutoSubmitResults)">
 					<Submit Files="#$(Platform) PGO Profile Results" FileType="+w" Description="Automated update of $(Platform) PGO files"/>
@@ -84,7 +84,7 @@
 			</Node>
 			
 			<!-- all done -->
-			<Node Name ="PGO Profile Replay $(Platform)" Requires="Gather PGO Profile Results $(Platform)">
+			<Node Name="$(PreNodeName)PGO Profile Replay $(Platform)" Requires="$(PreNodeName)Gather PGO Profile Results $(Platform)">
 				<Log Message="PGO Profiling completed for $(Platform)"/>
 			</Node>
 		</Agent>
