Fix buildgraph CookedEditor requirement errors

diff --git a/Engine/Build/Graph/Tasks/BuildAndTestProject.xml b/Engine/Build/Graph/Tasks/BuildAndTestProject.xml
--- a/Engine/Build/Graph/Tasks/BuildAndTestProject.xml
+++ b/Engine/Build/Graph/Tasks/BuildAndTestProject.xml
@@ -277,14 +277,14 @@
 			<!-- Stage CookedEditor -->
 			<Property Name="CookedEditorStageArgs" Value="-skipcook -stage -pak -platform=$(EditorPlatform) -clientconfig=$(CookedEditorConfigurations) -StagingDirectory=&quot;$(ProjectOutputDirectory)&quot;" />
 			<Property Name="CookedEditorStageRequirements" Value="$(CookedEditorCookRequirements)"/>
-			<Property Name="CookedEditorStageRequirements" Value="$(CookedEditorStageRequirements);Cook CookedEditor $(EditorPlatform)" If="!$(SkipCook)"/>
+			<Property Name="CookedEditorStageRequirements" Value="$(CookedEditorStageRequirements);$(PreNodeName)Cook CookedEditor $(EditorPlatform)" If="!$(SkipCook)"/>
 
-			<Node Name="$(PreNodeName)Stage CookedEditor $(EditorPlatform)" Requires="$(CookedEditorStageRequirements)">
+			<Node Name="$(PreNodeName)Stage CookedEditor $(EditorPlatform)" Requires="$(CookedEditorStageRequirements)">	
 				<Command Name="MakeCookedEditor" Arguments="-project=$(TargetProject) $(CookedEditorStageArgs)"/>
 			</Node>
 
 			<!-- Publish CookedEditor -->
-			<Node Name="$(PreNodeName)Publish CookedEditor $(EditorPlatform)" Requires="Stage CookedEditor $(EditorPlatform)">
+			<Node Name="$(PreNodeName)Publish CookedEditor $(EditorPlatform)" Requires="$(PreNodeName)Stage CookedEditor $(EditorPlatform)">
 				<Log Message="Copying to $(NetworkOutputDirectory)/$(EditorPlatformPath)CookedEditor"/>
 				<Do If="!$(SkipPublish)">
 					<Copy Files="..." From="$(ProjectOutputDirectory)/$(EditorPlatformPath)CookedEditor" To="$(NetworkOutputDirectory)/$(EditorPlatformPath)CookedEditor"/>
