Support Annotations on test runs

diff --git a/Engine/Build/Graph/Tasks/BuildAndTestProject.xml b/Engine/Build/Graph/Tasks/BuildAndTestProject.xml
--- a/Engine/Build/Graph/Tasks/BuildAndTestProject.xml
+++ b/Engine/Build/Graph/Tasks/BuildAndTestProject.xml
@@ -98,6 +98,11 @@
 	<Option Name="PathSettingsFile" DefaultValue="Inc/PathSettings.xml" Description="Path settings to use" />
 	<Option Name="GauntletSettingsFile" DefaultValue="Inc/GauntletSettings.xml" Description="Gauntlet settings to use" />
 
+    <!-- Annotations notif settings -->
+	<Option Name="AnnotationsTarget" DefaultValue="" Description="Set annotations target" />
+	<Option Name="AnnotationsPlatforms" DefaultValue="" Description="List of platforms that require annotations"/>
+	<Option Name="AnnotationsTests" DefaultValue="" Description="List of tests that require annotations"/>
+
 	<!-- Some common properties -->
 	<Include Script="Inc/CommonProperties.xml" />
 
@@ -213,7 +218,10 @@
 				<Property Name="GauntletReportArgs" Value="-logdir=&quot;$(OutputReportPath)&quot; -artifactname=Gauntlet -ReportExportPath=&quot;$(OutputReportPath)&quot;" />
 				<Property Name="GauntletReportArgs" Value="$(GauntletReportArgs) -ReportURL=&quot;$(AutomationReportUrl)/$(OutputFragmentPath)/&quot;" If="'$(AutomationReportUrl)' != ''" />
 			</Do>
-			<Node Name="$(PlatformBootTestNodeName)" Requires="$(GeneratedToolNodes)">
+			<!-- Annotations -->
+			<Property Name="NodeAnnotations" Value=""/>
+			<Property Name="NodeAnnotations" Value="Workflow=$(AnnotationsTarget)" If="ContainsItem('$(AnnotationsPlatforms)', '$(EditorPlatform)Editor', ';') and ContainsItem('$(AnnotationsTests)', 'BootTest', ';')"/>
+			<Node Name="$(PlatformBootTestNodeName)" Requires="$(GeneratedToolNodes)" Annotations="$(NodeAnnotations)">
 				<Command Name="RunUnreal" Arguments="-project=$(TargetProject) -build=Editor -test=&quot;UE.EditorBootTest&quot; $(GauntletEditorTestArgs) $(GauntletReportArgs)" />
 			</Node>
 			<Property Name="ProjectEditorDDC" Value="$(PreNodeName)Editor DDC $(EditorPlatform)" />
@@ -235,7 +243,9 @@
 					<Property Name="GauntletReportArgs" Value="$(GauntletReportArgs) -ReportURL=&quot;$(AutomationReportUrl)/$(OutputFragmentPath)/&quot;" If="'$(AutomationReportUrl)' != ''" />
 				</Do>
 				<Property Name="TestNodeName" Value="$(PreNodeName)Editor $(EditorPlatform) Test=$(TestName)" />
-				<Node Name="$(TestNodeName)" Requires="$(TestPrerequisites)">
+				<Property Name="NodeAnnotations" Value=""/>
+				<Property Name="NodeAnnotations" Value="Workflow=$(AnnotationsTarget)" If="ContainsItem('$(AnnotationsPlatforms)', '$(EditorPlatform)Editor', ';') and ContainsItem('$(AnnotationsTests)', '$(TestName)', ';')"/>
+				<Node Name="$(TestNodeName)" Requires="$(TestPrerequisites)" Annotations="$(NodeAnnotations)">
 					<Command Name="RunUnreal" Arguments="-project=$(TargetProject) -build=Editor -test=&quot;$(TestName)&quot; $(GauntletReportArgs) $(GauntletEditorTestArgs)" />
 				</Node>
 				<!-- Add each test node to our generated list -->
@@ -313,7 +323,9 @@
 				<Property Name="GauntletReportArgs" Value="-logdir=&quot;$(OutputReportPath)&quot; -artifactname=Gauntlet -ReportExportPath=&quot;$(OutputReportPath)&quot;" />
 				<Property Name="GauntletReportArgs" Value="$(GauntletReportArgs) -ReportURL=&quot;$(AutomationReportUrl)/$(OutputFragmentPath)/&quot;" If="'$(AutomationReportUrl)' != ''" />
 			</Do>
-			<Node Name="$(PlatformBootTestNodeName)" Requires="$(TestPrerequisites)">
+			<Property Name="NodeAnnotations" Value=""/>
+			<Property Name="NodeAnnotations" Value="Workflow=$(AnnotationsTarget)" If="ContainsItem('$(AnnotationsPlatforms)', '$(EditorPlatform)Editor', ';') and ContainsItem('$(AnnotationsTests)', 'BootTest', ';')"/>
+			<Node Name="$(PlatformBootTestNodeName)" Requires="$(TestPrerequisites)" Annotations="$(NodeAnnotations)">
 				<Command Name="RunUnreal" Arguments="-project=$(TargetProject) -build=&quot;$(NetworkOutputDirectory)/$(EditorPlatformPath)CookedEditor&quot; -test=&quot;UE.EditorBootTest&quot; $(GauntletCookedEditorTestArgs) $(CustomModuleRoleArg) $(GauntletReportArgs)" />
 			</Node>
 
@@ -331,7 +343,9 @@
 					<Property Name="GauntletReportArgs" Value="$(GauntletReportArgs) -ReportURL=&quot;$(AutomationReportUrl)/$(OutputFragmentPath)/&quot;" If="'$(AutomationReportUrl)' != ''" />
 				</Do>
 				<Property Name="TestNodeName" Value="$(PreNodeName)CookedEditor $(EditorPlatform) Test=$(TestName)" />
-				<Node Name="$(TestNodeName)" Requires="$(TestPrerequisites)">
+				<Property Name="NodeAnnotations" Value=""/>
+				<Property Name="NodeAnnotations" Value="Workflow=$(AnnotationsTarget)" If="ContainsItem('$(AnnotationsPlatforms)', '$(EditorPlatform)Editor', ';') and ContainsItem('$(AnnotationsTests)', '$(TestName)', ';')"/>
+				<Node Name="$(TestNodeName)" Requires="$(TestPrerequisites)" Annotations="$(NodeAnnotations)">
 					<Command Name="RunUnreal" Arguments="-project=$(TargetProject) -build=&quot;$(NetworkOutputDirectory)/$(EditorPlatformPath)CookedEditor&quot; -test=&quot;$(TestName)&quot; $(GauntletReportArgs) $(CustomModuleRoleArg) $(GauntletCookedEditorTestArgs)" />
 				</Node>
 				<!-- Add each test node to our generated list -->
@@ -524,9 +538,11 @@
 				<Property Name="GauntletTestArgs" Value="$(GauntletTargetTestArgs)"/>
 				<Property Name="GauntletTestArgs" Value="$(GauntletTargetTestAltArgs)" If="ContainsItem('$(AlternateDevicePoolPlatforms)', '$(TargetPlatform)', ';')"/>
 
-				<!-- Declare a boot test for the target-->	
+				<!-- Declare a boot test for the target-->
+				<Property Name="NodeAnnotations" Value=""/>
+				<Property Name="NodeAnnotations" Value="Workflow=$(AnnotationsTarget)" If="ContainsItem('$(AnnotationsPlatforms)', '$(TargetPlatform)', ';') and ContainsItem('$(AnnotationsTests)', 'BootTest', ';')"/>
 				<Property Name="BootTestName" Value="$(PreNodeName)$(TargetPlatform) $(TestConfiguration) BootTest" />
-				<Node Name="$(BootTestName)" After="$(PlatformPublishPackagedNodeName)" Requires="$(TestPrerequisites)">
+				<Node Name="$(BootTestName)" After="$(PlatformPublishPackagedNodeName)" Requires="$(TestPrerequisites)" Annotations="$(NodeAnnotations)">
 					<Do If="'$(AutomationReportOutputDirectory)' != ''">
 						<Property Name="OutputFragmentPath" Value="$(TargetPlatform)/$(TestConfiguration)/UE.BootTest" />
 						<Property Name="OutputReportPath" Value="$(AutomationReportOutputDirectory)/$(OutputFragmentPath)" />
@@ -552,8 +568,10 @@
 						<Property Name="GauntletReportArgs" Value="$(GauntletReportArgs) -ReportURL=&quot;$(AutomationReportUrl)/$(OutputFragmentPath)/&quot;" If="'$(AutomationReportUrl)' != ''" />
 					</Do>
 					<Property Name="TestNodeName" Value="$(PreNodeName)$(TargetPlatform) $(TestConfiguration) $(TestName)" />
+					<Property Name="NodeAnnotations" Value=""/>
+					<Property Name="NodeAnnotations" Value="Workflow=$(AnnotationsTarget)" If="ContainsItem('$(AnnotationsPlatforms)', '$(TargetPlatform)', ';') and ContainsItem('$(AnnotationsTests)', '$(TestName)', ';')"/>
 					<!-- Run the test. No 'Requires' as we run on the same machine locally and use the publish folder otherwise -->
-					<Node Name="$(TestNodeName)" Requires="$(TestPrerequisites)">
+					<Node Name="$(TestNodeName)" Requires="$(TestPrerequisites)" Annotations="$(NodeAnnotations)">
 						<Command Name="RunUnreal" Arguments="-project=$(TargetProject) -test=&quot;$(TestName)&quot; -platform=$(TargetPlatform) -configuration=$(TestConfiguration) $(GauntletBuildArg) $(GauntletTestArgs) $(GauntletReportArgs)"/>
 					</Node>
 					<!-- Add it to our list -->
