Sanitize test name for node title, mainly because of Group syntax

diff --git a/Engine/Build/Graph/Tasks/BuildAndTestProject.xml b/Engine/Build/Graph/Tasks/BuildAndTestProject.xml
--- a/Engine/Build/Graph/Tasks/BuildAndTestProject.xml
+++ b/Engine/Build/Graph/Tasks/BuildAndTestProject.xml
@@ -149,6 +149,17 @@
 	<Property Name="TargetType" Value="Game" />
 	<Property Name="TargetType" Value="Client" If="Contains('$(TargetName)', 'Client')" />
 
+	<!-- Sanitize node title for test name that include ':' ie: 'Runtest=Group:Animation' => 'Runtest=Group_Animation' -->
+	<Macro Name="SanitizeTestNameTitle">
+		<!-- Expect a Property named "TestNameTitle" to be defined beforehand -->
+		<Do If="Contains('$(TestNameTitle)', ':')">
+			<Property Name="Left" Value=""/>
+			<Property Name="Right" Value=""/>
+			<Regex Pattern="(.*):(.*)" Capture="Left;Right" Input="$(TestNameTitle)"/>
+			<Property Name="TestNameTitle" Value="$(Left)_$(Right)"/>
+		</Do>
+	</Macro>
+
 	<!-- These properties will be filled in with nodes that are generated based on options -->
 	<Property Name="GeneratedToolNodes" Value="" />
 	<Property Name="GeneratedEditorTestNodes" Value="" />
@@ -236,13 +247,15 @@
 			</Do>
 			<!-- Now generate a test node for each editor test in the list-->
 			<ForEach Name="TestName" Values="$(EditorTestList)" Separator="+">
+				<Property Name="TestNameTitle" Value="$(TestName)"/>
+				<Expand Name="SanitizeTestNameTitle"/>
 				<Do If="'$(AutomationReportOutputDirectory)' != ''">
-					<Property Name="OutputFragmentPath" Value="$(EditorPlatform)/Editor$(TestName)" />
+					<Property Name="OutputFragmentPath" Value="$(EditorPlatform)/Editor$(TestNameTitle)" />
 					<Property Name="OutputReportPath" Value="$(AutomationReportOutputDirectory)/$(OutputFragmentPath)" />
 					<Property Name="GauntletReportArgs" Value="-logdir=&quot;$(OutputReportPath)&quot; -artifactname=Gauntlet -ReportExportPath=&quot;$(OutputReportPath)&quot;" />
 					<Property Name="GauntletReportArgs" Value="$(GauntletReportArgs) -ReportURL=&quot;$(AutomationReportUrl)/$(OutputFragmentPath)/&quot;" If="'$(AutomationReportUrl)' != ''" />
 				</Do>
-				<Property Name="TestNodeName" Value="$(PreNodeName)Editor $(EditorPlatform) Test=$(TestName)" />
+				<Property Name="TestNodeName" Value="$(PreNodeName)Editor $(EditorPlatform) Test=$(TestNameTitle)" />
 				<Property Name="NodeAnnotations" Value=""/>
 				<Property Name="NodeAnnotations" Value="Workflow=$(AnnotationsTarget)" If="ContainsItem('$(AnnotationsPlatforms)', '$(EditorPlatform)Editor', ';') and ContainsItem('$(AnnotationsTests)', '$(TestName)', ';')"/>
 				<Node Name="$(TestNodeName)" Requires="$(TestPrerequisites)" Annotations="$(NodeAnnotations)">
@@ -336,13 +349,15 @@
 			<!-- Now generate a test node for each editor test in the list-->
 			<Property Name="PlatformCookedEditorTestNodes" Value="" />
 			<ForEach Name="TestName" Values="$(CookedEditorTestList)" Separator="+">
+				<Property Name="TestNameTitle" Value="$(TestName)"/>
+				<Expand Name="SanitizeTestNameTitle"/>
 				<Do If="'$(AutomationReportOutputDirectory)' != ''">
-					<Property Name="OutputFragmentPath" Value="$(EditorPlatform)/Editor$(TestName)" />
+					<Property Name="OutputFragmentPath" Value="$(EditorPlatform)/Editor$(TestNameTitle)" />
 					<Property Name="OutputReportPath" Value="$(AutomationReportOutputDirectory)/$(OutputFragmentPath)" />
 					<Property Name="GauntletReportArgs" Value="-logdir=&quot;$(OutputReportPath)&quot; -artifactname=Gauntlet -ReportExportPath=&quot;$(OutputReportPath)&quot;" />
 					<Property Name="GauntletReportArgs" Value="$(GauntletReportArgs) -ReportURL=&quot;$(AutomationReportUrl)/$(OutputFragmentPath)/&quot;" If="'$(AutomationReportUrl)' != ''" />
 				</Do>
-				<Property Name="TestNodeName" Value="$(PreNodeName)CookedEditor $(EditorPlatform) Test=$(TestName)" />
+				<Property Name="TestNodeName" Value="$(PreNodeName)CookedEditor $(EditorPlatform) Test=$(TestNameTitle)" />
 				<Property Name="NodeAnnotations" Value=""/>
 				<Property Name="NodeAnnotations" Value="Workflow=$(AnnotationsTarget)" If="ContainsItem('$(AnnotationsPlatforms)', '$(EditorPlatform)Editor', ';') and ContainsItem('$(AnnotationsTests)', '$(TestName)', ';')"/>
 				<Node Name="$(TestNodeName)" Requires="$(TestPrerequisites)" Annotations="$(NodeAnnotations)">
@@ -560,14 +575,16 @@
 				</Do>
 				<!-- Go through the list of tests -->	
 				<ForEach Name="TestName" Values="$(TargetTestList)" Separator="+">
+					<Property Name="TestNameTitle" Value="$(TestName)"/>
+					<Expand Name="SanitizeTestNameTitle"/>
 					<!-- Run and report on the test-->
 					<Do If="'$(AutomationReportOutputDirectory)' != ''">
-						<Property Name="OutputFragmentPath" Value="$(TargetPlatform)/$(TestConfiguration)/$(TestName)" />
+						<Property Name="OutputFragmentPath" Value="$(TargetPlatform)/$(TestConfiguration)/$(TestNameTitle)" />
 						<Property Name="OutputReportPath" Value="$(AutomationReportOutputDirectory)/$(OutputFragmentPath)" />
 						<Property Name="GauntletReportArgs" Value="-logdir=&quot;$(OutputReportPath)&quot; -artifactname=Gauntlet -ReportExportPath=&quot;$(OutputReportPath)&quot;" />
 						<Property Name="GauntletReportArgs" Value="$(GauntletReportArgs) -ReportURL=&quot;$(AutomationReportUrl)/$(OutputFragmentPath)/&quot;" If="'$(AutomationReportUrl)' != ''" />
 					</Do>
-					<Property Name="TestNodeName" Value="$(PreNodeName)$(TargetPlatform) $(TestConfiguration) $(TestName)" />
+					<Property Name="TestNodeName" Value="$(PreNodeName)$(TargetPlatform) $(TestConfiguration) $(TestNameTitle)" />
 					<Property Name="NodeAnnotations" Value=""/>
 					<Property Name="NodeAnnotations" Value="Workflow=$(AnnotationsTarget)" If="ContainsItem('$(AnnotationsPlatforms)', '$(TargetPlatform)', ';') and ContainsItem('$(AnnotationsTests)', '$(TestName)', ';')"/>
 					<!-- Run the test. No 'Requires' as we run on the same machine locally and use the publish folder otherwise -->
