Adding an option to BuildAndTestProject graph to specify which configuration to build UnrealInsights for, opening up the option to use Development or Shipping.

diff --git a/Engine/Build/Graph/Tasks/BuildAndTestProject.xml b/Engine/Build/Graph/Tasks/BuildAndTestProject.xml
--- a/Engine/Build/Graph/Tasks/BuildAndTestProject.xml
+++ b/Engine/Build/Graph/Tasks/BuildAndTestProject.xml
@@ -34,6 +34,7 @@
 	<Option Name="TestAgentOverride" DefaultValue="" Description="If set, other logic is ignored and this agent type is used for all test work." />
 	<Option Name="TargetPlatformTestAgentOverride" DefaultValue="DevkitAutomation" Description="Override agent type use to run target platform tests"/>
 	<Option Name="ClientCompileAgentOverride" DefaultValue="" Description="Specifies the agent type to use for C++ compilation jobs for target platform binaries." />
+	<Option Name="UnrealInsightsBuildConfiguration" DefaultValue="Shipping" Restrict="Development|Debug|Shipping" Description="In which configuration to build Unreal Insights, if included. Shipping, Debug, or Development." />
 	<Option Name="IncludeUnrealInsightsInToolCompile" DefaultValue="false" Description="Include Unreal Insights into Compile Tools prerequisites" />
 	<Option Name="IncludeEpicWebHelperInToolCompile" DefaultValue="false" Description="Include EpicWebHelper into Compile Tools prerequisites" />
 
@@ -201,7 +202,7 @@
 						</ForEach>
 					</Do>
 				</Do>
-				<Compile Target="UnrealInsights" Platform="$(EditorPlatform)" Configuration="Shipping" Arguments="$(GenericCompileArguments)" If="$(IncludeUnrealInsightsInToolCompile)"/>
+				<Compile Target="UnrealInsights" Platform="$(EditorPlatform)" Configuration="$(UnrealInsightsBuildConfiguration)" Arguments="$(GenericCompileArguments)" If="$(IncludeUnrealInsightsInToolCompile)"/>
 				<Compile Target="EpicWebHelper" Platform="$(EditorPlatform)" Configuration="Shipping" Arguments="$(GenericCompileArguments)" If="$(IncludeEpicWebHelperInToolCompile)"/>
 			</Node>
 			<Node Name="$(PreNodeName)Compile Editor $(EditorPlatform)" Requires="$(PreNodeName)Compile Tools $(EditorPlatform)">
