Fix BuildEditorAndTools not supporting project arguments with relative paths correctly.

diff --git a/Engine/Build/Graph/Examples/BuildEditorAndTools.xml b/Engine/Build/Graph/Examples/BuildEditorAndTools.xml
--- a/Engine/Build/Graph/Examples/BuildEditorAndTools.xml
+++ b/Engine/Build/Graph/Examples/BuildEditorAndTools.xml
@@ -80,13 +80,8 @@
 		</Node>
 		
 		<!-- Compile the editor executable -->
-		<Node Name="Compile $(EditorTarget) Win64" Requires="Compile Tools Win64" Produces="#EditorBinaries;#TargetInfo">
-			<Property Name="UProjectParam" Value=""/>
-			<Property Name="UProjectParam" Value="-Project=&quot;$(UProjectPath)&quot;" If="'$(UProjectPath)' != ''"/>
-			<Compile Target="$(EditorTarget)" Platform="Win64" Configuration="Development" Tag="#EditorBinaries" Arguments="$(UProjectParam)"/>
-
-			<Property Name="UProjectArg" Value=""/>
-			<Property Name="UProjectArg" Value="$(UProjectPath)" If="'$(UProjectPath)' != ''"/>
+		<Node Name="Compile $(EditorTarget) Win64" Produces="#EditorBinaries;#TargetInfo">
+			<Compile Target="$(EditorTarget)" Platform="Win64" Configuration="Development" Project="$(UProjectPath)" Tag="#EditorBinaries"/>
 			<QueryTargets ProjectFile="$(UProjectPath)" Tag="#TargetInfo"/>
 		</Node>
 		
