Seperate properties for BuildVersion and BuildNamePath

#rnx

diff --git a/Engine/Build/Graph/Tasks/BuildAndTestProject.xml b/Engine/Build/Graph/Tasks/BuildAndTestProject.xml
--- a/Engine/Build/Graph/Tasks/BuildAndTestProject.xml
+++ b/Engine/Build/Graph/Tasks/BuildAndTestProject.xml
@@ -454,14 +454,14 @@
 			<Node Name="$(PlatformCompileNodeName)" Requires="$(PlatformToolsNodeName)" Produces="#$(PlatformCompileNodeName) Binaries">
 				<Switch>
 					<Case If="$(IsContentOnlyProject)">
-						<Command Name="BuildGame" Arguments="$(BCRArgs) -configuration=$(TargetConfigurations) -notools -ubtargs=&quot;-BuildVersion=$(BuildName) $(GenericCompileArguments) $(ExtraTargetCompileArguments)&quot;" />
+						<Command Name="BuildGame" Arguments="$(BCRArgs) -configuration=$(TargetConfigurations) -notools -ubtargs=&quot;-BuildVersion=$(BuildVersion) $(GenericCompileArguments) $(ExtraTargetCompileArguments)&quot;" />
 						<Tag Files="$(RootDir)/$(ProjectPath)/Binaries/..." Except="$(ProjectPath)/Binaries/DotNET/..." With="#$(PlatformCompileNodeName) Binaries"/>
 					</Case>
 					<Default>
 						<ForEach Name="TargetConfiguration" Values="$(TargetConfigurations)" Separator="+">
-							<Compile Target="$(TargetName)" Platform="$(TargetPlatform)" Configuration="$(TargetConfiguration)" Arguments="-BuildVersion=&quot;$(BuildName)&quot; $(GenericCompileArguments) $(ExtraTargetCompileArguments)"/>
+							<Compile Target="$(TargetName)" Platform="$(TargetPlatform)" Configuration="$(TargetConfiguration)" Arguments="-BuildVersion=&quot;$(BuildVersion)&quot; $(GenericCompileArguments) $(ExtraTargetCompileArguments)"/>
 							<Do If="'$(DedicatedServer)'" >
-								<Compile Target="$(TargetName)Server" Platform="Win64" Configuration="$(TargetConfiguration)" Arguments="-BuildVersion=&quot;$(BuildName)&quot; $(GenericCompileArguments) $(ExtraTargetCompileArguments)"/>
+								<Compile Target="$(TargetName)Server" Platform="Win64" Configuration="$(TargetConfiguration)" Arguments="-BuildVersion=&quot;$(BuildVersion)&quot; $(GenericCompileArguments) $(ExtraTargetCompileArguments)"/>
 							</Do>
 						</ForEach>
 					</Default>
@@ -472,7 +472,7 @@
 			<Do If="'$(UploadSymbolsNodeName)' != ''">
 				<Node Name="$(UploadSymbolsNodeName)" Requires="$(PlatformCompileNodeName)">
 					<Log Message="Uploading symbols to &quot;$(SymbolServerPath)&quot; ..."/>
-					<SymStore Platform="$(TargetPlatform)" Files="#$(PlatformCompileNodeName)" StoreDir="$(SymbolServerPath)/$(TargetPlatform)" Product="$(TargetName)" BuildVersion="$(BuildName)" />
+					<SymStore Platform="$(TargetPlatform)" Files="#$(PlatformCompileNodeName)" StoreDir="$(SymbolServerPath)/$(TargetPlatform)" Product="$(TargetName)" BuildVersion="$(BuildVersion)" />
 				</Node>
 			</Do>
 
diff --git a/Engine/Build/Graph/Tasks/Inc/CommonProperties.xml b/Engine/Build/Graph/Tasks/Inc/CommonProperties.xml
--- a/Engine/Build/Graph/Tasks/Inc/CommonProperties.xml
+++ b/Engine/Build/Graph/Tasks/Inc/CommonProperties.xml
@@ -18,6 +18,7 @@
 	<Property Name="PreflightSuffix" Value=""/>
 	<Property Name="PreflightSuffix" Value="-PF-$(PreflightChange)-$(UE_HORDE_JOBID)" If="$(IsPreflight)"/>
 	<Property Name="CL-String" Value="CL-$(Change)$(PreflightSuffix)"/>
-	<Property Name="BuildName" Value="$(ProjectName)-$(CL-String)"/>
+	<Property Name="BuildNamePath" Value="$(ProjectName)-$(CL-String)"/>
+	<Property Name="BuildVersion" Value="$(ProjectName)-CL-$(Change)"/>
 
 </BuildGraph>
diff --git a/Engine/Build/Graph/Tasks/Inc/PathSettings.xml b/Engine/Build/Graph/Tasks/Inc/PathSettings.xml
--- a/Engine/Build/Graph/Tasks/Inc/PathSettings.xml
+++ b/Engine/Build/Graph/Tasks/Inc/PathSettings.xml
@@ -76,7 +76,7 @@
 	<Property Name="NetworkOutputDirectory" Value="" />
 	<Property Name="NetworkOutputDirectory" Value="$(NetworkTempRoot)" If="$(IsBuildMachine)"/>
 	<Property Name="NetworkOutputDirectory" Value="$(NetworkPublishRoot)" If="$(PublishBuilds) or $(PublishHashedStagingData)"/>
-	<Property Name="NetworkOutputDirectory" Value="$(NetworkOutputDirectory)/$(EscapedBranch)/$(BuildName)" If="'$(NetworkOutputDirectory)' != ''"/>
+	<Property Name="NetworkOutputDirectory" Value="$(NetworkOutputDirectory)/$(EscapedBranch)/$(BuildNamePath)" If="'$(NetworkOutputDirectory)' != ''"/>
 
 	<!-- Automation can publish reports to a folder that is mapped to a website -->
 	<Do If="$(IsBuildMachine) or $(PublishReports)" >
