Disable CI for platform

#rnx

diff --git a/Engine/Build/Graph/Tasks/Inc/PlatformOptions.xml b/Engine/Build/Graph/Tasks/Inc/PlatformOptions.xml
--- a/Engine/Build/Graph/Tasks/Inc/PlatformOptions.xml
+++ b/Engine/Build/Graph/Tasks/Inc/PlatformOptions.xml
@@ -27,7 +27,7 @@
 <Property Name="ShowExplicitFlags" Value="true" />
 
 <Property Name="AllEditorPlatforms" Value="Mac+Win64+Linux" />
-<Property Name="AllTargetPlatforms" Value="Android+HoloLens+iOS+Mac+PS4+PS5+Switch+tvOS+Win64+WinGDK+XboxOneGDK+XB1+XSX" />
+<Property Name="AllTargetPlatforms" Value="Android+iOS+Mac+PS4+PS5+Switch+tvOS+Win64+WinGDK+XboxOneGDK+XB1+XSX" />
 
 <Option Name="EditorPlatforms" DefaultValue="" Description="Alternative to 'WithFooEditor', a list of editors to build and test. Win64+Mac etc"/>
 <Option Name="TargetPlatforms" DefaultValue="" Description="Alternative to 'WithFoo', a list of Platforms to build and test. PS4+XSX etc"/>
@@ -41,7 +41,7 @@
 
 <!-- No means of reliably running these tests atm -->
 <Option Name="EditorPlatformsToSkipTests" DefaultValue="" Restrict=".*" Description="List of editor platforms that editor tests won't be run on" />
-<Option Name="TargetPlatformsToSkipTests" DefaultValue="HoloLens+IOS+Linux+TVOS" Restrict=".*" Description="List of target platforms that target tests won't be run on" />
+<Option Name="TargetPlatformsToSkipTests" DefaultValue="IOS+Linux+TVOS" Restrict=".*" Description="List of target platforms that target tests won't be run on" />
 <Option Name="CookedEditorPlatformsToSkipTests" DefaultValue="" Restrict=".*" Description="List of editor cooked platforms that editor tests won't be run on" />
 <Option Name="PlatformsToSkipTests" DefaultValue="" Restrict=".*" Description="List of platforms that all tests won't be run on. Overrides Editor/TargetPlatformsToSkipTests when non-empty." />
 
@@ -56,7 +56,6 @@
 	<Option Name="WithLinuxCookedEditor" DefaultValue="false" Description="Build and cook Linux CookedEditor"/>
 
 	<Option Name="WithAndroid" DefaultValue="false" Description="Build and test Android"/>
-	<Option Name="WithHoloLens" DefaultValue="false" Description="Build and test HoloLens"/>
 	<Option Name="WithIOS" DefaultValue="false" Description="Build and test IOS"/>
 	<Option Name="WithLinux" DefaultValue="false" Description="Build and test Linux"/>
 	<Option Name="WithMac" DefaultValue="false" Description="Build and test Mac"/>
@@ -87,7 +86,6 @@
 	<Do If="'$(TargetPlatforms)' == ''" >
 		<Property Name="TargetPlatforms" Value="" />
 		<Property Name="TargetPlatforms" Value="$(TargetPlatforms)Android+" If="$(WithAndroid)" />
-		<Property Name="TargetPlatforms" Value="$(TargetPlatforms)HoloLens+" If="$(WithHoloLens)" />
 		<Property Name="TargetPlatforms" Value="$(TargetPlatforms)IOS+" If="$(WithIOS)" />
 		<Property Name="TargetPlatforms" Value="$(TargetPlatforms)Linux+" If="$(WithLinux)" />
 		<Property Name="TargetPlatforms" Value="$(TargetPlatforms)Mac+" If="$(WithMac)" />
diff --git a/Engine/Build/InstalledEngineFilters.xml b/Engine/Build/InstalledEngineFilters.xml
--- a/Engine/Build/InstalledEngineFilters.xml
+++ b/Engine/Build/InstalledEngineFilters.xml
@@ -20,7 +20,6 @@
 	<Property Name="IOSStripFilter" Value="*.a;"/>
 	<Property Name="TVOSStripFilter" Value="*.a;"/>
 	<Property Name="LinuxStripFilter" Value="Engine\Binaries\Linux\*."/>
-	<Property Name="HoloLensStripFilter" Value="*.pdb"/>
 
 
 	<Property Name="WindowsSignFilter" Value="*.exe;*.dll"/>
diff --git a/Engine/Build/LowLevelTests.xml b/Engine/Build/LowLevelTests.xml
--- a/Engine/Build/LowLevelTests.xml
+++ b/Engine/Build/LowLevelTests.xml
@@ -59,7 +59,7 @@
 
 	<Option Name="GauntletDevicePool" DefaultValue="UE4" Description="Gauntlet Device Pool to use to run target platform tests"/>
 	<Option Name="AlternateGauntletDevicePool" DefaultValue="UE5" Description="Alternate Gauntlet Device Pool to use to run target platform tests"/>
-	<Option Name="AlternateDevicePoolPlatforms" DefaultValue="$(AlternateDevicePoolPlatformsProperty);HoloLens" Description="Platforms that need to use the Alternate Gauntlet Device Pool"/>
+	<Option Name="AlternateDevicePoolPlatforms" DefaultValue="$(AlternateDevicePoolPlatformsProperty)" Description="Platforms that need to use the Alternate Gauntlet Device Pool"/>
 	<Option Name="Device" DefaultValue="" Description="Device address to run tests on"/>
 	<Option Name="Configuration" DefaultValue="Development" Description="Configuration type: development, debug, test, shipping etc."/>
 	<Option Name="AttachToDebugger" DefaultValue="false" Description="Wait for debugger to attach."/>
diff --git a/Engine/Build/LowLevelTests/Foundation.xml b/Engine/Build/LowLevelTests/Foundation.xml
--- a/Engine/Build/LowLevelTests/Foundation.xml
+++ b/Engine/Build/LowLevelTests/Foundation.xml
@@ -33,7 +33,4 @@
   <Property Name="LinuxFoundationRunSupported" Value="True" />
   <Property Name="LinuxArm64FoundationRunSupported" Value="True" />
   <Property Name="TVOSFoundationRunSupported" Value="True" />
-  <Property Name="HoloLensFoundationTags" Value="" />
-  <Property Name="HoloLensFoundationExtraCompilationArgs" Value="-allmodules" />
-  <Property Name="HoloLensFoundationRunSupported" Value="True" />
 </BuildGraph>
\ No newline at end of file
diff --git a/Engine/Build/LowLevelTests/General.xml b/Engine/Build/LowLevelTests/General.xml
--- a/Engine/Build/LowLevelTests/General.xml
+++ b/Engine/Build/LowLevelTests/General.xml
@@ -5,17 +5,14 @@
 	<Property Name="MacExtraArgs" Value=""/>
 	<Property Name="IOSExtraArgs" Value=""/>
 	<Property Name="AndroidExtraArgs" Value=""/>
-	<Property Name="HoloLensExtraArgs" Value=""/>
 	<Property Name="Win64RequiresSpecialAgent" Value="False"/>
 	<Property Name="LinuxRequiresSpecialAgent" Value="False"/>
 	<Property Name="MacRequiresSpecialAgent" Value="False"/>
 	<Property Name="IOSRequiresSpecialAgent" Value="False"/>
 	<Property Name="AndroidRequiresSpecialAgent" Value="False"/>
-	<Property Name="HoloLensRequiresSpecialAgent" Value="False"/>
 	<Property Name="Win64RequiresCatch2BuildExtraCMakeArgs" Value="False"/>
 	<Property Name="LinuxRequiresCatch2BuildExtraCMakeArgs" Value="False"/>
 	<Property Name="MacRequiresCatch2BuildExtraCMakeArgs" Value="False"/>
 	<Property Name="IOSRequiresCatch2BuildExtraCMakeArgs" Value="False"/>
 	<Property Name="AndroidRequiresCatch2BuildExtraCMakeArgs" Value="False"/>
-	<Property Name="HoloLensRequiresCatch2BuildExtraCMakeArgs" Value="False"/>
 </BuildGraph>
\ No newline at end of file
diff --git a/Engine/Build/LowLevelTests/HeadlessChaos.xml b/Engine/Build/LowLevelTests/HeadlessChaos.xml
--- a/Engine/Build/LowLevelTests/HeadlessChaos.xml
+++ b/Engine/Build/LowLevelTests/HeadlessChaos.xml
@@ -33,7 +33,4 @@
   <Property Name="LinuxHeadlessChaosRunSupported" Value="True" />
   <Property Name="LinuxArm64HeadlessChaosRunSupported" Value="True" />
   <Property Name="TVOSHeadlessChaosRunSupported" Value="True" />
-  <Property Name="HoloLensHeadlessChaosTags" Value="" />
-  <Property Name="HoloLensHeadlessChaosExtraCompilationArgs" Value="" />
-  <Property Name="HoloLensHeadlessChaosRunSupported" Value="True" />
 </BuildGraph>
\ No newline at end of file
diff --git a/Engine/Build/LowLevelTests/NetworkPredictionPlugin.xml b/Engine/Build/LowLevelTests/NetworkPredictionPlugin.xml
--- a/Engine/Build/LowLevelTests/NetworkPredictionPlugin.xml
+++ b/Engine/Build/LowLevelTests/NetworkPredictionPlugin.xml
@@ -33,7 +33,4 @@
   <Property Name="LinuxNetworkPredictionPluginRunSupported" Value="True" />
   <Property Name="LinuxArm64NetworkPredictionPluginRunSupported" Value="True" />
   <Property Name="TVOSNetworkPredictionPluginRunSupported" Value="True" />
-  <Property Name="HoloLensNetworkPredictionPluginTags" Value="" />
-  <Property Name="HoloLensNetworkPredictionPluginExtraCompilationArgs" Value="" />
-  <Property Name="HoloLensNetworkPredictionPluginRunSupported" Value="True" />
 </BuildGraph>
\ No newline at end of file
diff --git a/Engine/Build/LowLevelTests/ReplicationSystem.xml b/Engine/Build/LowLevelTests/ReplicationSystem.xml
--- a/Engine/Build/LowLevelTests/ReplicationSystem.xml
+++ b/Engine/Build/LowLevelTests/ReplicationSystem.xml
@@ -33,7 +33,4 @@
   <Property Name="LinuxReplicationSystemRunSupported" Value="True" />
   <Property Name="LinuxArm64ReplicationSystemRunSupported" Value="True" />
   <Property Name="TVOSReplicationSystemRunSupported" Value="True" />
-  <Property Name="HoloLensReplicationSystemTags" Value="" />
-  <Property Name="HoloLensReplicationSystemExtraCompilationArgs" Value="" />
-  <Property Name="HoloLensReplicationSystemRunSupported" Value="True" />
 </BuildGraph>
\ No newline at end of file
diff --git a/Engine/Build/PhysXBuild.xml b/Engine/Build/PhysXBuild.xml
--- a/Engine/Build/PhysXBuild.xml
+++ b/Engine/Build/PhysXBuild.xml
@@ -12,7 +12,6 @@
 	<Option Name="TargetConfigs" DefaultValue="profile+release+checked" Description="Target Configs"/>
 	<Option Name="TargetWindowsCompilers" DefaultValue="VisualStudio2019" Description="Target Compilers"/>
 	<Option Name="WithWin64" DefaultValue="false" Description="Compile for Win64"/>
-	<Option Name="WithHoloLens" DefaultValue="false" Description="Compile for HoloLens"/>
 	<Option Name="WithPS4" DefaultValue="false" Description="Compile for PS4"/>
 	<Option Name="WithXboxOne" DefaultValue="false" Description="Compile for Xbox One"/>
 	<Option Name="WithXboxOneGDK" DefaultValue="$(WithXboxOne)" Description="Compile for Xbox One"/>
@@ -44,11 +43,6 @@
 		</Node>
 	</Agent>
 	<!-- @MIXEDREALITY_CHANGE : BEGIN ARM -->
-	<Agent Name="PhysX HoloLens" Type="EngineWin64;PhysX_CompileWin64;Win64">
-		<Node Name="Compile PhysX HoloLens">
-			<Command Name="BuildPhysX" Arguments="-TargetPlatforms=ARM64 -TargetConfigs=&quot;$(TargetConfigs)&quot; -TargetWindowsCompilers=&quot;$(TargetWindowsCompilers)&quot; $(OtherArgs)"/>
-		</Node>
-	</Agent>
 	<Agent Name="PhysX PS4" Type="EngineWin64;PhysX_CompileWin64;Win64">
 		<Node Name="Compile PhysX PS4">
 			<Command Name="BuildPhysX" Arguments="-TargetPlatforms=PS4 -TargetConfigs=&quot;$(TargetConfigs)&quot; -TargetWindowsCompilers=&quot;$(TargetWindowsCompilers)&quot; $(OtherArgs)"/>
@@ -112,7 +106,6 @@
 
 	<Property Name="PhysXTargets" Value=""/>
 	<Property Name="PhysXTargets" Value="$(PhysXTargets);Compile PhysX Win64" If="$(WithWin64)"/>
-	<Property Name="PhysXTargets" Value="$(PhysXTargets);Compile PhysX HoloLens" If="$(WithHoloLens)"/>
 	<Property Name="PhysXTargets" Value="$(PhysXTargets);Compile PhysX PS4" If="$(WithPS4)"/>
 	<Property Name="PhysXTargets" Value="$(PhysXTargets);Compile PhysX XboxOneGDK" If="$(WithXboxOneGDK)"/>
 	<Property Name="PhysXTargets" Value="$(PhysXTargets);Compile PhysX Linux" If="$(WithLinux)"/>
