Fix Content only editor build + use TagReceipt to gather build product for content only project.

#rnx

diff --git a/Engine/Build/Graph/Tasks/BuildAndTestProject.xml b/Engine/Build/Graph/Tasks/BuildAndTestProject.xml
--- a/Engine/Build/Graph/Tasks/BuildAndTestProject.xml
+++ b/Engine/Build/Graph/Tasks/BuildAndTestProject.xml
@@ -119,9 +119,7 @@
 	<Error Message="Project file $(TargetProject) does not exist" If="!Exists('$(TargetProject)')" />
 
 	<!-- If the user didn't supply a TargetName then use the ProjectName -->
-	<Property Name="TargetName" Value="UnrealGame" If="$(IsContentOnlyProject) And '$(TargetName)' == ''"/>
 	<Property Name="TargetName" Value="$(ProjectName)" If="'$(TargetName)' == ''"/>
-	<Property Name="EditorTargetName" Value="UnrealEditor" If="$(IsContentOnlyProject) And '$(EditorTargetName)' == ''"/>
 	<Property Name="EditorTargetName" Value="$(ProjectName)Editor" If="'$(EditorTargetName)' == ''"/>
 	<Property Name="CookedEditorTargetName" Value="$(ProjectName)CookedEditor" If="'$(CookedEditorTargetName)' == ''"/>
 
@@ -205,8 +203,18 @@
 				<Compile Target="UnrealInsights" Platform="$(EditorPlatform)" Configuration="$(UnrealInsightsBuildConfiguration)" Arguments="$(GenericCompileArguments)" If="$(IncludeUnrealInsightsInToolCompile)"/>
 				<Compile Target="EpicWebHelper" Platform="$(EditorPlatform)" Configuration="Shipping" Arguments="$(GenericCompileArguments)" If="$(IncludeEpicWebHelperInToolCompile)"/>
 			</Node>
-			<Node Name="$(PreNodeName)Compile Editor $(EditorPlatform)" Requires="$(PreNodeName)Compile Tools $(EditorPlatform)">
-				<Compile Target="$(EditorTargetName)" Platform="$(EditorPlatform)" Configuration="Development" Arguments="$(GenericCompileArguments) $(ExtraEditorCompileArguments)"/>
+			<Property Name="CompileEditor" Value="$(PreNodeName)Compile Editor $(EditorPlatform)"/>
+			<Node Name="$(CompileEditor)" Requires="$(PreNodeName)Compile Tools $(EditorPlatform)">
+				<Switch>
+					<Case If="$(IsContentOnlyProject)">
+						<Command Name="BuildCookRun" Arguments="-project=&quot;$(RootDir)/$(TargetProject)&quot; -platform=$(EditorPlatform) -configuration=Development -build -notools -ubtargs=&quot;$(GenericCompileArguments) $(ExtraEditorCompileArguments)&quot;"/>
+						<Tag Files="$(RootDir)/$(ProjectPath)/Binaries/..." Filter="$(EditorTargetName).target" With="#$(CompileEditor)"/>
+						<TagReceipt Files="#$(CompileEditor)" BuildProducts="true" With="#$(CompileEditor)"/>
+					</Case>
+					<Default>
+						<Compile Target="$(EditorTargetName)" Platform="$(EditorPlatform)" Configuration="Development" Arguments="$(GenericCompileArguments) $(ExtraEditorCompileArguments)"/>
+					</Default>
+				</Switch>
 			</Node>
 			<!-- Cooked editor target -->
 			<Node Name="$(PreNodeName)Compile CookedEditor $(EditorPlatform)" Requires="$(PreNodeName)Compile Tools $(EditorPlatform)">
@@ -215,7 +223,7 @@
 				</ForEach>
 			</Node>
 			<!--Add these nodes to our dependency list -->
-			<Property Name="GeneratedToolNodes" Value="$(PreNodeName)Compile Editor $(EditorPlatform);$(PreNodeName)Compile Tools $(EditorPlatform)" If="!$(SkipBuild)" />
+			<Property Name="GeneratedToolNodes" Value="$(CompileEditor);$(PreNodeName)Compile Tools $(EditorPlatform)" If="!$(SkipBuild)" />
 			<!-- Editor Labels -->
 			<Label Category="Editors" Name="$(PreLabelName)Editor $(EditorPlatform)" Requires="$(GeneratedToolNodes)" />
 		</Agent>
@@ -233,7 +241,7 @@
 			<!-- Annotations -->
 			<Property Name="NodeAnnotations" Value=""/>
 			<Property Name="NodeAnnotations" Value="Workflow=$(AnnotationsTarget)" If="ContainsItem('$(AnnotationsPlatforms)', '$(EditorPlatform)Editor', ';') and ContainsItem('$(AnnotationsTests)', 'BootTest', ';')"/>
-			<Node Name="$(PlatformBootTestNodeName)" Requires="$(GeneratedToolNodes)" Annotations="$(NodeAnnotations)">
+			<Node Name="$(PlatformBootTestNodeName)" Requires="$(TestPrerequisites)" Annotations="$(NodeAnnotations)">
 				<Command Name="RunUnreal" Arguments="-project=$(TargetProject) -build=Editor -test=&quot;UE.EditorBootTest&quot; $(GauntletEditorTestArgs) $(GauntletReportArgs)" />
 			</Node>
 			<Property Name="ProjectEditorDDC" Value="$(PreNodeName)Editor DDC $(EditorPlatform)" />
@@ -451,11 +459,12 @@
 		<Agent Name="Compile $(PlatformAgentName)" Type="$(CompileAgentType)">
 
 			<!-- Compile -->
-			<Node Name="$(PlatformCompileNodeName)" Requires="$(PlatformToolsNodeName)" Produces="#$(PlatformCompileNodeName) Binaries">
+			<Node Name="$(PlatformCompileNodeName)" Requires="$(PlatformToolsNodeName)">
 				<Switch>
 					<Case If="$(IsContentOnlyProject)">
 						<Command Name="BuildGame" Arguments="$(BCRArgs) -configuration=$(TargetConfigurations) -notools -ubtargs=&quot;-BuildVersion=$(BuildVersion) $(GenericCompileArguments) $(ExtraTargetCompileArguments)&quot;" />
-						<Tag Files="$(RootDir)/$(ProjectPath)/Binaries/..." Except="$(ProjectPath)/Binaries/DotNET/..." With="#$(PlatformCompileNodeName) Binaries"/>
+						<Tag Files="$(RootDir)/$(ProjectPath)/Binaries/..." Filter="$(TargetName).target" With="#$(PlatformCompileNodeName)"/>
+						<TagReceipt Files="#$(PlatformCompileNodeName)" BuildProducts="true" With="#$(PlatformCompileNodeName)"/>
 					</Case>
 					<Default>
 						<ForEach Name="TargetConfiguration" Values="$(TargetConfigurations)" Separator="+">
@@ -489,7 +498,7 @@
 			</Node>
 			
 			<!-- Stage -->
-			<Node Name="$(PlatformStageNodeName)" Requires="$(PlatformCompileNodeName);#$(PlatformCompileNodeName) Binaries;#$(PlatformCookNodeName) Complete">
+			<Node Name="$(PlatformStageNodeName)" Requires="$(PlatformCompileNodeName);#$(PlatformCookNodeName) Complete">
 				<Switch>
 					<Case If="$(IsMobilePlatform)">
 						<ForEach Name="PackageConfiguration" Values="$(PackageConfigurations)" Separator="+">
