Use -VerifyLogin when running UE.Automation and UE.BootTest test on clients (2nd try)

#rnx

diff --git a/Engine/Build/Graph/Tasks/Inc/GauntletSettings.xml b/Engine/Build/Graph/Tasks/Inc/GauntletSettings.xml
--- a/Engine/Build/Graph/Tasks/Inc/GauntletSettings.xml
+++ b/Engine/Build/Graph/Tasks/Inc/GauntletSettings.xml
@@ -20,6 +20,7 @@
 	<Option Name="GauntletDevicePool" DefaultValue="UE4" Description="Gauntlet Device Pool to use to run target platform tests"/>
 	<Option Name="AlternateGauntletDevicePool" DefaultValue="UE5" Description="Alternate Gauntlet Device Pool to use to run target platform tests"/>
 	<Option Name="AlternateDevicePoolPlatforms" DefaultValue="XSX;PS5" Description="Platforms that need to use the Alternate Gauntlet Device Pool"/>
+	<Option Name="EnableDeviceLoginVerification" DefaultValue="$(IsBuildMachine)" Description="Verify a user is login on Devices that supports it, try to sign in if not."/>
 	<Option Name="ExtraGauntletCommonArgs" DefaultValue="" Description="Additional Common arguments passed to Gauntlet command line"/>
 	<Option Name="ExtraGauntletTargetTestArgs" DefaultValue="" Description="Additional Target test arguments passed to Gauntlet command line"/>
 	<Option Name="ExtraGauntletEditorTestArgs" DefaultValue="" Description="Additional Editor test arguments passed to Gauntlet command line"/>
@@ -29,6 +30,7 @@
 
 	<Property Name="GauntletDeviceArgs" Value="" />
 	<Property Name="GauntletDeviceArgs" Value="-deviceurl=https://horde.devtools.epicgames.com/ -reboot -cleardevices -sandbox=&quot;$(ProjectName)&quot;" If="$(IsBuildMachine)" />
+	<Property Name="GauntletDeviceArgs" Value="$(GauntletDeviceArgs) -VerifyLogin" If="$(EnableDeviceLoginVerification)"/>
 	<Property Name="GauntletDeviceAltArgs" Value="$(GauntletDeviceArgs)" />
 	<Property Name="GauntletDeviceArgs" Value="$(GauntletDeviceArgs) -devicepool=&quot;$(GauntletDevicePool)&quot;" If="$(IsBuildMachine)"/>
 	<Property Name="GauntletDeviceAltArgs" Value="$(GauntletDeviceAltArgs) -devicepool=&quot;$(AlternateGauntletDevicePool)&quot;" If="$(IsBuildMachine)"/>
diff --git a/Engine/Source/Programs/AutomationTool/Gauntlet/Unreal/Base/Gauntlet.UnrealTestConfiguration.cs b/Engine/Source/Programs/AutomationTool/Gauntlet/Unreal/Base/Gauntlet.UnrealTestConfiguration.cs
--- a/Engine/Source/Programs/AutomationTool/Gauntlet/Unreal/Base/Gauntlet.UnrealTestConfiguration.cs
+++ b/Engine/Source/Programs/AutomationTool/Gauntlet/Unreal/Base/Gauntlet.UnrealTestConfiguration.cs
@@ -943,7 +943,12 @@ namespace Gauntlet
 				}
 			}
 
-			if(CommandUtils.IsBuildMachine)
+			if (ConfigRole.RoleType.IsClient())
+			{
+				VerifyLogin = Globals.Params.ParseParam("VerifyLogin");
+			}
+
+			if (CommandUtils.IsBuildMachine)
 			{
 				AppConfig.CommandLineParams.AddUnique("BUILDMACHINE");
 			}
diff --git a/Engine/Source/Programs/AutomationTool/Gauntlet/Unreal/Game/EpicGame.TestConfig.cs b/Engine/Source/Programs/AutomationTool/Gauntlet/Unreal/Game/EpicGame.TestConfig.cs
--- a/Engine/Source/Programs/AutomationTool/Gauntlet/Unreal/Game/EpicGame.TestConfig.cs
+++ b/Engine/Source/Programs/AutomationTool/Gauntlet/Unreal/Game/EpicGame.TestConfig.cs
@@ -245,6 +245,10 @@ namespace EpicGame
 						UserAccount.ApplyToConfig(AppConfig);
 					}
 				}
+				else
+				{
+					VerifyLogin = false;
+				}
 			}
 
 			if (ConfigRole.RoleType.IsClient())
