Ensure we don't use snapshot export on a platform that doesn't use project stores.  Fix more localhost issues when calling the external utility.

diff --git a/Engine/Build/Graph/Tasks/BuildAndTestProject.xml b/Engine/Build/Graph/Tasks/BuildAndTestProject.xml
--- a/Engine/Build/Graph/Tasks/BuildAndTestProject.xml
+++ b/Engine/Build/Graph/Tasks/BuildAndTestProject.xml
@@ -525,7 +525,7 @@
 
 			<!-- (Optional) Publish to Cloud -->
 			<Property Name="PlatformPublishStageDependencyName" Value="$(PlatformHashNodeName)" />
-			<Do If="$(PublishBuilds) and $(PublishCloudSnapshot) and '$(UE-CloudPublishHost)' != '' and '$(UE-CloudPublishNamespace)' != ''">
+			<Do If="$(PublishBuilds) and $(PublishCloudSnapshot) and '$(UE-CloudPublishHost)' != '' and '$(UE-CloudPublishNamespace)' != '' and Exists('$(ProjectPath)/Saved/$(CookPlatform)/.projectstore')">
 				<Property Name="PlatformCloudPublishNodeName" Value="$(PreNodeName)Cloud Publish $(TargetPlatform)" />
 				<Node Name="$(PlatformCloudPublishNodeName)" Requires="$(PlatformHashNodeName)">
 					<ZenExportSnapshot Project="$(TargetProject)" Platform="$(CookPlatform)" DestinationStorageType="Cloud" DestinationCloudHost="$(UE-CloudPublishHost)" DestinationCloudNamespace="$(UE-CloudPublishNamespace)" DestinationCloudIdentifier="$(EscapedBranch)/$(BuildNamePath)" SnapshotDescriptorFile="$(NetworkOutputDirectory)/$(StagedPlatformFolder)/SnapshotDescriptor.json"/>
diff --git a/Engine/Source/Programs/AutomationTool/BuildGraph/Tasks/ZenExportSnapshotTask.cs b/Engine/Source/Programs/AutomationTool/BuildGraph/Tasks/ZenExportSnapshotTask.cs
--- a/Engine/Source/Programs/AutomationTool/BuildGraph/Tasks/ZenExportSnapshotTask.cs
+++ b/Engine/Source/Programs/AutomationTool/BuildGraph/Tasks/ZenExportSnapshotTask.cs
@@ -215,7 +215,11 @@ namespace AutomationTool.Tasks
 						ExportKeyIds[ExportIndex] = DestinationKeyBuilder.ToString().ToLowerInvariant();
 						IoHash DestinationKeyHash = IoHash.Compute(Encoding.UTF8.GetBytes(ExportKeyIds[ExportIndex]));
 
-						if (!ExportSource.IsLocalHost)
+						if (ExportSource.IsLocalHost)
+						{
+							SingleSourceCommandline.AppendFormat(" --hosturl http://localhost:{0}", ExportSource.HostPort);
+						}
+						else
 						{
 							SingleSourceCommandline.AppendFormat(" --hosturl http://{0}:{1}", ExportSource.HostName, ExportSource.HostPort);
 						}
