Fix error in path when performing existence check on project store file before exporting snapshot.

diff --git a/Engine/Build/Graph/Tasks/BuildAndTestProject.xml b/Engine/Build/Graph/Tasks/BuildAndTestProject.xml
--- a/Engine/Build/Graph/Tasks/BuildAndTestProject.xml
+++ b/Engine/Build/Graph/Tasks/BuildAndTestProject.xml
@@ -525,7 +525,7 @@
 
 			<!-- (Optional) Publish to Cloud -->
 			<Property Name="PlatformPublishStageDependencyName" Value="$(PlatformHashNodeName)" />
-			<Do If="$(PublishBuilds) and $(PublishCloudSnapshot) and '$(UE-CloudPublishHost)' != '' and '$(UE-CloudPublishNamespace)' != '' and Exists('$(ProjectPath)/Saved/$(CookPlatform)/.projectstore')">
+			<Do If="$(PublishBuilds) and $(PublishCloudSnapshot) and '$(UE-CloudPublishHost)' != '' and '$(UE-CloudPublishNamespace)' != '' and Exists('$(ProjectPath)/Saved/Cooked/$(CookPlatform)/.projectstore')">
 				<Property Name="PlatformCloudPublishNodeName" Value="$(PreNodeName)Cloud Publish $(TargetPlatform)" />
 				<Node Name="$(PlatformCloudPublishNodeName)" Requires="$(PlatformHashNodeName)">
 					<ZenExportSnapshot Project="$(TargetProject)" Platform="$(CookPlatform)" DestinationStorageType="Cloud" DestinationCloudHost="$(UE-CloudPublishHost)" DestinationCloudNamespace="$(UE-CloudPublishNamespace)" DestinationCloudIdentifier="$(EscapedBranch)/$(BuildNamePath)" SnapshotDescriptorFile="$(NetworkOutputDirectory)/$(StagedPlatformFolder)/SnapshotDescriptor.json"/>
