Move optional cloud publish operation into the Stage operation temporarily until some changes are made related to zenserver lifetime.

diff --git a/Engine/Build/Graph/Tasks/BuildAndTestProject.xml b/Engine/Build/Graph/Tasks/BuildAndTestProject.xml
--- a/Engine/Build/Graph/Tasks/BuildAndTestProject.xml
+++ b/Engine/Build/Graph/Tasks/BuildAndTestProject.xml
@@ -514,6 +514,10 @@
 						<Command Name="BuildCookRun" Arguments="$(BCRArgs) -configuration=$(PackageConfigurations) $(StageAndPackageArguments) -skipbuild -skipcook -stage -pak" />
 					</Default>
 				</Switch>
+				<!-- (Optional) Publish to Cloud - temporary location -->
+				<Do If="$(PublishBuilds) and $(PublishCloudSnapshot) and '$(UE-CloudPublishHost)' != '' and '$(UE-CloudPublishNamespace)' != ''">
+					<ZenExportSnapshot Project="$(TargetProject)" Platform="$(CookPlatform)" DestinationStorageType="Cloud" DestinationCloudHost="$(UE-CloudPublishHost)" DestinationCloudNamespace="$(UE-CloudPublishNamespace)" DestinationCloudIdentifier="$(EscapedBranch)/$(BuildNamePath)" SnapshotDescriptorFile="$(NetworkOutputDirectory)/$(StagedPlatformFolder)/SnapshotDescriptor.json"/>
+				</Do>
 			</Node>
 			
 			<!-- Hash -->
@@ -525,13 +529,13 @@
 
 			<!-- (Optional) Publish to Cloud -->
 			<Property Name="PlatformPublishStageDependencyName" Value="$(PlatformHashNodeName)" />
-			<Do If="$(PublishBuilds) and $(PublishCloudSnapshot) and '$(UE-CloudPublishHost)' != '' and '$(UE-CloudPublishNamespace)' != ''">
+			<!--<Do If="$(PublishBuilds) and $(PublishCloudSnapshot) and '$(UE-CloudPublishHost)' != '' and '$(UE-CloudPublishNamespace)' != ''">
 				<Property Name="PlatformCloudPublishNodeName" Value="$(PreNodeName)Cloud Publish $(TargetPlatform)" />
 				<Node Name="$(PlatformCloudPublishNodeName)" Requires="$(PlatformHashNodeName)">
 					<ZenExportSnapshot Project="$(TargetProject)" Platform="$(CookPlatform)" DestinationStorageType="Cloud" DestinationCloudHost="$(UE-CloudPublishHost)" DestinationCloudNamespace="$(UE-CloudPublishNamespace)" DestinationCloudIdentifier="$(EscapedBranch)/$(BuildNamePath)" SnapshotDescriptorFile="$(NetworkOutputDirectory)/$(StagedPlatformFolder)/SnapshotDescriptor.json"/>
 				</Node>
 				<Property Name="PlatformPublishStageDependencyName" Value="$(PlatformCloudPublishNodeName)" />
-			</Do>
+			</Do>-->
 			
 			<!-- Publish (Staged) -->
 			<Node Name="$(PlatformPublishStagedNodeName)" Requires="$(PlatformPublishStageDependencyName)">
diff --git a/Engine/Source/Programs/AutomationTool/BuildGraph/Tasks/ZenExportSnapshotTask.cs b/Engine/Source/Programs/AutomationTool/BuildGraph/Tasks/ZenExportSnapshotTask.cs
--- a/Engine/Source/Programs/AutomationTool/BuildGraph/Tasks/ZenExportSnapshotTask.cs
+++ b/Engine/Source/Programs/AutomationTool/BuildGraph/Tasks/ZenExportSnapshotTask.cs
@@ -128,7 +128,7 @@ namespace AutomationTool.Tasks
 			FileReference ProjectFile = Parameters.Project;
 			if(!FileReference.Exists(ProjectFile))
 			{
-				throw new AutomationException("Missing project file - {0}", ProjectFile.FullName);
+				return Task.CompletedTask;
 			}
 
 			List<ExportSourceData> ExportSources = new List<ExportSourceData>();
