Cooking: IterativeValidate: Add a BuildGraph script to run the iterativevalidate cook steps.
#rnx

diff --git a/Engine/Build/Graph/Tests/IterativeValidate.xml b/Engine/Build/Graph/Tests/IterativeValidate.xml
new file mode 100644
--- /dev/null
+++ b/Engine/Build/Graph/Tests/IterativeValidate.xml
@@ -0,0 +1,41 @@
+<?xml version='1.0' ?>
+<BuildGraph xmlns="http://www.epicgames.com/BuildGraph" xmlns:xsi="http://www.w3.org/2001/XMLSchema-instance" xsi:schemaLocation="http://www.epicgames.com/BuildGraph ../Schema.xsd" >
+
+	<!-- Project Options -->
+	<Option Name="Project" DefaultValue="_Unknown_" Restrict=".+" Description="The projects to cook"/>
+	<!-- Note we can not make the default='', because the default is used for some validation. -->
+	<Error Message="No project specified" If="'$(Project)' == '_Unknown_'"/>
+
+	<!-- Platform Options -->
+	<Option Name="TargetPlatform" DefaultValue="Windows" Restrict=".+" Description="The platform to cook"/>
+
+	<!-- Iterative Options -->
+	<Option Name="IterativeMode" DefaultValue="incremental" Restrict="iterative|incremental" Description="Whether to iteratively cook using Incremental or using the legacy iterative"/>
+	<Property Name="IterativeArgs" Value="-zenstore -ini:Editor:[CookSettings]:HybridIterativeEnabled=true" If="'$(IterativeMode)' == 'incremental'"/>
+	<Property Name="IterativeArgs" Value="-iterate -ini:Editor:[CookSettings]:HybridIterativeEnabled=false -ini:Editor:[CookSettings]:IterativeUseClassFilters=true" If="'$(IterativeMode)' == 'iterative'"/>
+
+	<!-- Agent=IterativeValidate variables -->
+	<Property Name="EditorPlatform" Value="Win64"/>
+	<Property Name="EditorTarget" Value="$(Project)Editor"/>
+	<!-- TODO: Remove these editordomain arguments once we enable it for all cooks -->
+	<Property Name="EditorDomainArgs" Value="-editordomain -noeditordomainread -noeditordomainwrite"/>
+	<Property Name="CommonCookArgs" Value="-targetplatform=$(TargetPlatform) $(EditorDomainArgs) $(IterativeArgs) -skiponlyeditoronly=true"/>
+	<Property Name="CompileNodeName" Value="Compile $(EditorPlatform) $(Project)"/>
+	<Property Name="CookIterativeValidatePrePassName" Value="Cook IterativeValidatePrePass $(Project) $(TargetPlatform)"/>
+	<Property Name="CookIterativeValidateName" Value="Cook IterativeValidate $(Project) $(TargetPlatform)"/>
+
+	<Agent Name="IterativeValidateAgent" Type="IterativeValidateAgent_$(Project)_$(TargetPlatform)">
+		<Node Name="$(CompileNodeName)">
+			<Compile Target="$(EditorTarget)" Platform="$(EditorPlatform)" Project="$(Project)" Configuration="Development"/>
+		</Node>
+		<Node Name="$(CookIterativeValidatePrePassName)" Requires="$(CompileNodeName)">
+			<Commandlet Name="Cook" Project="$(Project)" Arguments="$(CommonCookArgs) -IterativeValidatePrePass"/>
+		</Node>
+		<Node Name="$(CookIterativeValidateName)" Requires="$(CookIterativeValidatePrePassName)">
+			<Commandlet Name="Cook" Project="$(Project)" Arguments="$(CommonCookArgs) -IterativeValidate"/>
+		</Node>
+	</Agent>
+
+	<Aggregate Name="IterativeValidate" Requires="$(CookIterativeValidateName)"/>
+
+</BuildGraph>
\ No newline at end of file
