IterativeValidate BuildGraph task: Add separate argument for EditorTarget and ProjectFile because they are not always derivable from projectname.
#rnx

diff --git a/Engine/Build/Graph/Tests/IterativeValidate.xml b/Engine/Build/Graph/Tests/IterativeValidate.xml
--- a/Engine/Build/Graph/Tests/IterativeValidate.xml
+++ b/Engine/Build/Graph/Tests/IterativeValidate.xml
@@ -2,9 +2,13 @@
 <BuildGraph xmlns="http://www.epicgames.com/BuildGraph" xmlns:xsi="http://www.w3.org/2001/XMLSchema-instance" xsi:schemaLocation="http://www.epicgames.com/BuildGraph ../Schema.xsd" >
 
 	<!-- Project Options -->
-	<Option Name="Project" DefaultValue="_Unknown_" Restrict=".+" Description="The projects to cook"/>
+	<Option Name="ProjectFile" DefaultValue="_Unknown_" Restrict=".+" Description="The projects to cook"/>
+	<Option Name="ProjectName" DefaultValue="_Unknown_" Restrict=".+" Description="The projects to cook"/>
+	<Option Name="EditorTarget" DefaultValue="$(ProjectName)Editor" Description="Compile name of the editor target"/>
+
 	<!-- Note we can not make the default='', because the default is used for some validation. -->
-	<Error Message="No project specified" If="'$(Project)' == '_Unknown_'"/>
+	<Error Message="No project file specified" If="'$(ProjectFile)' == '_Unknown_'"/>
+	<Error Message="No project name specified" If="'$(ProjectName)' == '_Unknown_'"/>
 
 	<!-- Platform Options -->
 	<Option Name="TargetPlatform" DefaultValue="Windows" Restrict=".+" Description="The platform to cook"/>
@@ -16,23 +20,22 @@
 
 	<!-- Agent=IterativeValidate variables -->
 	<Property Name="EditorPlatform" Value="Win64"/>
-	<Property Name="EditorTarget" Value="$(Project)Editor"/>
 	<!-- TODO: Remove these editordomain arguments once we enable it for all cooks -->
 	<Property Name="EditorDomainArgs" Value="-editordomain -noeditordomainread -noeditordomainwrite"/>
 	<Property Name="CommonCookArgs" Value="-targetplatform=$(TargetPlatform) $(EditorDomainArgs) $(IterativeArgs) -skiponlyeditoronly=true"/>
-	<Property Name="CompileNodeName" Value="Compile $(EditorPlatform) $(Project)"/>
-	<Property Name="CookIterativeValidatePrePassName" Value="Cook IterativeValidatePrePass $(Project) $(TargetPlatform)"/>
-	<Property Name="CookIterativeValidateName" Value="Cook IterativeValidate $(Project) $(TargetPlatform)"/>
+	<Property Name="CompileNodeName" Value="Compile $(EditorPlatform) $(ProjectName)"/>
+	<Property Name="CookIterativeValidatePrePassName" Value="Cook IterativeValidatePrePass $(ProjectName) $(TargetPlatform)"/>
+	<Property Name="CookIterativeValidateName" Value="Cook IterativeValidate $(ProjectName) $(TargetPlatform)"/>
 
-	<Agent Name="IterativeValidateAgent" Type="IterativeValidateAgent_$(Project)_$(TargetPlatform)">
+	<Agent Name="IterativeValidateAgent" Type="IterativeValidateAgent_$(ProjectName)_$(TargetPlatform)">
 		<Node Name="$(CompileNodeName)">
-			<Compile Target="$(EditorTarget)" Platform="$(EditorPlatform)" Project="$(Project)" Configuration="Development"/>
+			<Compile Target="$(EditorTarget)" Platform="$(EditorPlatform)" Project="$(ProjectFile)" Configuration="Development"/>
 		</Node>
 		<Node Name="$(CookIterativeValidatePrePassName)" Requires="$(CompileNodeName)">
-			<Commandlet Name="Cook" Project="$(Project)" Arguments="$(CommonCookArgs) -IterativeValidatePrePass"/>
+			<Commandlet Name="Cook" Project="$(ProjectFile)" Arguments="$(CommonCookArgs) -IterativeValidatePrePass"/>
 		</Node>
 		<Node Name="$(CookIterativeValidateName)" Requires="$(CookIterativeValidatePrePassName)">
-			<Commandlet Name="Cook" Project="$(Project)" Arguments="$(CommonCookArgs) -IterativeValidate"/>
+			<Commandlet Name="Cook" Project="$(ProjectFile)" Arguments="$(CommonCookArgs) -IterativeValidate"/>
 		</Node>
 	</Agent>
 
