IterativeValidate BuildGraph: Add separate argument for AgentType so it can use Win64 instead of Windows as the platform name.
#rnx

diff --git a/Engine/Build/Graph/Tests/IterativeValidate.xml b/Engine/Build/Graph/Tests/IterativeValidate.xml
--- a/Engine/Build/Graph/Tests/IterativeValidate.xml
+++ b/Engine/Build/Graph/Tests/IterativeValidate.xml
@@ -13,6 +13,10 @@
 	<!-- Platform Options -->
 	<Option Name="TargetPlatform" DefaultValue="Windows" Restrict=".+" Description="The platform to cook"/>
 
+	<!-- Agent Options -->
+	<Option Name="AgentType" DefaultValue="" Description="The platform to cook"/>
+	<Property Name="AgentType" Value="IterativeValidateAgent_$(ProjectName)_$(TargetPlatform)" If="'$(AgentType)' == ''"/>
+
 	<!-- Iterative Options -->
 	<Option Name="IterativeMode" DefaultValue="incremental" Restrict="iterative|incremental" Description="Whether to iteratively cook using Incremental or using the legacy iterative"/>
 	<Property Name="IterativeArgs" Value="-zenstore -ini:Editor:[CookSettings]:HybridIterativeEnabled=true" If="'$(IterativeMode)' == 'incremental'"/>
@@ -27,7 +31,7 @@
 	<Property Name="CookIterativeValidatePrePassName" Value="Cook IterativeValidatePrePass $(ProjectName) $(TargetPlatform)"/>
 	<Property Name="CookIterativeValidateName" Value="Cook IterativeValidate $(ProjectName) $(TargetPlatform)"/>
 
-	<Agent Name="IterativeValidateAgent" Type="IterativeValidateAgent_$(ProjectName)_$(TargetPlatform)">
+	<Agent Name="IterativeValidateAgent" Type="$(AgentType)">
 		<Node Name="$(CompileNodeName)">
 			<Compile Target="$(EditorTarget)" Platform="$(EditorPlatform)" Project="$(ProjectFile)" Configuration="Development"/>
 		</Node>
