IterativeValidate - Add compilation of ShaderCompileWorker because it is required for cooking.
#rnx

diff --git a/Engine/Build/Graph/Tests/IterativeValidate.xml b/Engine/Build/Graph/Tests/IterativeValidate.xml
--- a/Engine/Build/Graph/Tests/IterativeValidate.xml
+++ b/Engine/Build/Graph/Tests/IterativeValidate.xml
@@ -33,6 +33,7 @@
 
 	<Agent Name="IterativeValidateAgent" Type="$(AgentType)">
 		<Node Name="$(CompileNodeName)">
+			<Compile Target="ShaderCompileWorker" Platform="$(EditorPlatform)" Project="$(ProjectFile)" Configuration="Development"/>
 			<Compile Target="$(EditorTarget)" Platform="$(EditorPlatform)" Project="$(ProjectFile)" Configuration="Development"/>
 		</Node>
 		<Node Name="$(CookIterativeValidatePrePassName)" Requires="$(CompileNodeName)">
