IterativeValidate - Force singleprocess cook since the data collected in each cook is not replicated from cookworkers.
#rnx

diff --git a/Engine/Build/Graph/Tests/IterativeValidate.xml b/Engine/Build/Graph/Tests/IterativeValidate.xml
--- a/Engine/Build/Graph/Tests/IterativeValidate.xml
+++ b/Engine/Build/Graph/Tests/IterativeValidate.xml
@@ -24,9 +24,12 @@
 
 	<!-- Agent=IterativeValidate variables -->
 	<Property Name="EditorPlatform" Value="Win64"/>
-	<!-- TODO: Remove these editordomain arguments once we enable it for all cooks -->
-	<Property Name="EditorDomainArgs" Value="-editordomain -noeditordomainread -noeditordomainwrite"/>
-	<Property Name="CommonCookArgs" Value="-targetplatform=$(TargetPlatform) $(EditorDomainArgs) $(IterativeArgs) -skiponlyeditoronly=true"/>
+	<Property Name="CommonCookArgs" Value="-targetplatform=$(TargetPlatform) $(IterativeArgs) -skiponlyeditoronly=true"/>
+	<!-- TODO: Remove these editordomain arguments once we enable editordomain for all cooks -->
+	<Property Name="CommonCookArgs" Value="$(CommonCookArgs) -editordomain -noeditordomainread -noeditordomainwrite"/>
+	<!-- Force singleprocess cook because IterativeValidatePackageWriter includes data that is not replicated -->
+	<Property Name="CommonCookArgs" Value="$(CommonCookArgs) -cookprocesscount=1"/>
+
 	<Property Name="CompileNodeName" Value="Compile $(EditorPlatform) $(ProjectName)"/>
 	<Property Name="CookIterativeValidatePrePassName" Value="Cook IterativeValidatePrePass $(ProjectName) $(TargetPlatform)"/>
 	<Property Name="CookIterativeValidateName" Value="Cook IterativeValidate $(ProjectName) $(TargetPlatform)"/>
