Add support for EditorCompileArgs Option
#rnx

diff --git a/Engine/Build/Graph/Examples/BuildEditorAndTools.xml b/Engine/Build/Graph/Examples/BuildEditorAndTools.xml
--- a/Engine/Build/Graph/Examples/BuildEditorAndTools.xml
+++ b/Engine/Build/Graph/Examples/BuildEditorAndTools.xml
@@ -38,6 +38,7 @@
 
 	<Option Name="EditorTarget" Restrict="[^ ]+" DefaultValue="UnrealEditor" Description="Name of the editor target to be built"/>
 	<Option Name="UProjectPath" Restrict="[^ ]*" DefaultValue="" Description="Optional Uproject path to include in the arguments"/>
+	<Option Name="EditorCompileArgs" DefaultValue="" Description="Arguments to be used for EditorTarget Compile"/>
 	<Option Name="GameTargets" Restrict="[^ ]*" DefaultValue="" Description="Names of game targets to build separated by semicolons, eg. UnrealClient;UnrealServer"/>
 	<Option Name="AdditionalTools" Restrict="[^ ]*" DefaultValue="" Description="Tool Compiles to build seperated by semicolons, eg. UnrealMultiUserServer;BuildPatchTool"/>
 	<Option Name="TargetPlatforms" Restrict="[^ ]*" DefaultValue="Win64" Description="List of the target platforms to build for, separated by semicolons, eg. Win64;Android"/>
@@ -81,7 +82,7 @@
 		
 		<!-- Compile the editor executable -->
 		<Node Name="Compile $(EditorTarget) Win64" Produces="#EditorBinaries;#TargetInfo">
-			<Compile Target="$(EditorTarget)" Platform="Win64" Configuration="Development" Project="$(UProjectPath)" Tag="#EditorBinaries"/>
+			<Compile Target="$(EditorTarget)" Platform="Win64" Configuration="Development" Project="$(UProjectPath)" Tag="#EditorBinaries" Arguments="$(EditorCompileArgs)"/>
 			<QueryTargets ProjectFile="$(UProjectPath)" Tag="#TargetInfo"/>
 		</Node>
 		
