Add options for Compile Tools Arguments
#rnx

diff --git a/Engine/Build/Graph/Examples/BuildEditorAndTools.xml b/Engine/Build/Graph/Examples/BuildEditorAndTools.xml
--- a/Engine/Build/Graph/Examples/BuildEditorAndTools.xml
+++ b/Engine/Build/Graph/Examples/BuildEditorAndTools.xml
@@ -24,6 +24,8 @@
 		Other options:
 			-set:EditorTarget=<TargetName>		Specifies the editor target to be built (defaults to UnrealEditor)
 			-set:UProjectPath=<ProjectPath>		Specifies the uproject path to pass as -Project in the compile args (defaults to none)
+			-set:EditorCompileArgs=<Args>		Specifies the arguments to pass to Editor Compile (defaults to none)
+			-set:ToolCompileArgs=<Args>			Specifies the arguments to pass to Tools Compile (defaults to none)
 			-set:GameTargets=<TargetName>		Specifies the game targets to be built (defaults to none)
 			-set:AdditionalTools=<A>;<B>;<C>..	Specifies the additional tool targets to be built (defaults to none)
 			-set:TargetPlatforms=<A>;<B>;<C>... Specifies the target platforms to build the game for (defaults to Win64)
@@ -39,6 +41,7 @@
 	<Option Name="EditorTarget" Restrict="[^ ]+" DefaultValue="UnrealEditor" Description="Name of the editor target to be built"/>
 	<Option Name="UProjectPath" Restrict="[^ ]*" DefaultValue="" Description="Optional Uproject path to include in the arguments"/>
 	<Option Name="EditorCompileArgs" DefaultValue="" Description="Arguments to be used for EditorTarget Compile"/>
+	<Option Name="ToolCompileArgs" DefaultValue="" Description="Arguments to be used for Tool Compile"/>
 	<Option Name="GameTargets" Restrict="[^ ]*" DefaultValue="" Description="Names of game targets to build separated by semicolons, eg. UnrealClient;UnrealServer"/>
 	<Option Name="AdditionalTools" Restrict="[^ ]*" DefaultValue="" Description="Tool Compiles to build seperated by semicolons, eg. UnrealMultiUserServer;BuildPatchTool"/>
 	<Option Name="TargetPlatforms" Restrict="[^ ]*" DefaultValue="Win64" Description="List of the target platforms to build for, separated by semicolons, eg. Win64;Android"/>
@@ -65,12 +68,12 @@
 		
 		<!-- Compile the tool executables -->
 		<Node Name="Compile Tools Win64" Requires="Update Version Files" Produces="#ToolBinaries">
-			<Compile Target="ShaderCompileWorker" Platform="Win64" Configuration="Development" Tag="#ToolBinaries"/>
-			<Compile Target="UnrealLightmass" Platform="Win64" Configuration="Development" Tag="#ToolBinaries"/>
-			<Compile Target="InterchangeWorker" Platform="Win64" Configuration="Development" Tag="#ToolBinaries"/>
-			<Compile Target="UnrealPak" Configuration="Development" Platform="Win64" Tag="#ToolBinaries"/>
+			<Compile Target="ShaderCompileWorker" Platform="Win64" Configuration="Development" Tag="#ToolBinaries" Arguments="$(ToolCompileArgs)"/>
+			<Compile Target="UnrealLightmass" Platform="Win64" Configuration="Development" Tag="#ToolBinaries" Arguments="$(ToolCompileArgs)"/>
+			<Compile Target="InterchangeWorker" Platform="Win64" Configuration="Development" Tag="#ToolBinaries" Arguments="$(ToolCompileArgs)"/>
+			<Compile Target="UnrealPak" Configuration="Development" Platform="Win64" Tag="#ToolBinaries" Arguments="$(ToolCompileArgs)"/>
 			<ForEach Name="AdditonalTool" Values="$(AdditionalTools)">
-				<Compile Target="$(AdditonalTool)" Platform="Win64" Configuration="Development" Tag="#ToolBinaries"/>
+				<Compile Target="$(AdditonalTool)" Platform="Win64" Configuration="Development" Tag="#ToolBinaries" Arguments="$(ToolCompileArgs)"/>
 			</ForEach>
 			
 			<!-- 
