Add ZenLaunch to the tools compile step in buildgraph so that it can be used from other build steps.

diff --git a/Engine/Build/Graph/Tasks/BuildAndTestProject.xml b/Engine/Build/Graph/Tasks/BuildAndTestProject.xml
--- a/Engine/Build/Graph/Tasks/BuildAndTestProject.xml
+++ b/Engine/Build/Graph/Tasks/BuildAndTestProject.xml
@@ -194,6 +194,7 @@
 			<!-- First generate nodes to build the required tools (e.g. UHT, ShaderCompiler, Editor -->
 			<Node Name="$(PreNodeName)Compile Tools $(EditorPlatform)" Requires="$(PreNodeName)Update Version Files $(EditorPlatform)">
 				<Compile Target="ShaderCompileWorker" Platform="$(EditorPlatform)" Configuration="Development" Arguments="$(GenericCompileArguments)"/>
+				<Compile Target="ZenLaunch" Platform="$(EditorPlatform)" Configuration="Development" Arguments="$(GenericCompileArguments)"/>
 				<Compile Target="UnrealPak" Platform="$(EditorPlatform)" Configuration="Development" Arguments="$(GenericCompileArguments)"/>
 				<Compile Target="CrashReportClient" Platform="$(EditorPlatform)" Configuration="Shipping" Arguments="$(GenericCompileArguments)"/>
 				<Compile Target="CrashReportClientEditor" Platform="$(EditorPlatform)" Configuration="Shipping" Arguments="$(GenericCompileArguments)"/>
