Pass Extra*CompileArguments to the CompileArguments property if they are incremental
#rnx

diff --git a/Engine/Build/Graph/Tasks/BuildAndTestProject.xml b/Engine/Build/Graph/Tasks/BuildAndTestProject.xml
--- a/Engine/Build/Graph/Tasks/BuildAndTestProject.xml
+++ b/Engine/Build/Graph/Tasks/BuildAndTestProject.xml
@@ -144,9 +144,9 @@
 	</Do>
 
 	<Do If="$(UseIncrementalAgents)" >
-		<Property Name="ExtraToolCompileArguments" Value="-allmodules -nodebuginfo -2019 -Strict -WarningsAsErrors"/>
-		<Property Name="ExtraEditorCompileArguments" Value="-allmodules -nodebuginfo -2019 -Strict -WarningsAsErrors"/>
-		<Property Name="ExtraTargetCompileArguments" Value="-nodebuginfo -WarningsAsErrors"/>
+		<Property Name="ExtraToolCompileArguments" Value="$(ExtraToolCompileArguments) -allmodules -nodebuginfo -2019 -Strict -WarningsAsErrors"/>
+		<Property Name="ExtraEditorCompileArguments" Value="$(ExtraEditorCompileArguments) -allmodules -nodebuginfo -2019 -Strict -WarningsAsErrors"/>
+		<Property Name="ExtraTargetCompileArguments" Value="$(ExtraTargetCompileArguments) -nodebuginfo -WarningsAsErrors"/>
 	</Do>
 
 	<!-- Detect the type of the client based on the target name -->
