Removing 2019 arg, should be using the default VS version

#rnx

diff --git a/Engine/Build/Graph/Tasks/BuildAndTestProject.xml b/Engine/Build/Graph/Tasks/BuildAndTestProject.xml
--- a/Engine/Build/Graph/Tasks/BuildAndTestProject.xml
+++ b/Engine/Build/Graph/Tasks/BuildAndTestProject.xml
@@ -144,8 +144,8 @@
 	</Do>
 
 	<Do If="$(UseIncrementalAgents)" >
-		<Property Name="ExtraToolCompileArguments" Value="$(ExtraToolCompileArguments) -allmodules -nodebuginfo -2019 -Strict -WarningsAsErrors"/>
-		<Property Name="ExtraEditorCompileArguments" Value="$(ExtraEditorCompileArguments) -allmodules -nodebuginfo -2019 -Strict -WarningsAsErrors"/>
+		<Property Name="ExtraToolCompileArguments" Value="$(ExtraToolCompileArguments) -allmodules -nodebuginfo -Strict -WarningsAsErrors"/>
+		<Property Name="ExtraEditorCompileArguments" Value="$(ExtraEditorCompileArguments) -allmodules -nodebuginfo -Strict -WarningsAsErrors"/>
 		<Property Name="ExtraTargetCompileArguments" Value="$(ExtraTargetCompileArguments) -nodebuginfo -WarningsAsErrors"/>
 	</Do>
 
