Adds branch and changelist to gauntlet jobs
Passes branch and changelist as metadata to CitySampleTest.PerformanceReport

diff --git a/Engine/Build/Graph/Tasks/Inc/GauntletSettings.xml b/Engine/Build/Graph/Tasks/Inc/GauntletSettings.xml
--- a/Engine/Build/Graph/Tasks/Inc/GauntletSettings.xml
+++ b/Engine/Build/Graph/Tasks/Inc/GauntletSettings.xml
@@ -24,9 +24,10 @@
 	<Option Name="ExtraGauntletCommonArgs" DefaultValue="" Description="Additional Common arguments passed to Gauntlet command line"/>
 	<Option Name="ExtraGauntletTargetTestArgs" DefaultValue="" Description="Additional Target test arguments passed to Gauntlet command line"/>
 	<Option Name="ExtraGauntletEditorTestArgs" DefaultValue="" Description="Additional Editor test arguments passed to Gauntlet command line"/>
-	
+
 	<Property Name="GauntletCommonArgs" Value="-maxduration=$(GauntletTimeOut) -unattended -verbose -ResumeOnCriticalFailure $(ExtraGauntletCommonArgs)" />
 	<Property Name="GauntletCommonArgs" Value="$(GauntletCommonArgs) -WriteTestResultsForHorde" If="'$(UE_HORDE_JOBID)' != ''"/>
+	<Property Name="GauntletCommonArgs" Value="$(GauntletCommonArgs) -branch=&quot;$(Branch)&quot; -changelist=&quot;$(Change)&quot;" />
 
 	<Property Name="GauntletDeviceArgs" Value="" />
 	<Property Name="GauntletDeviceArgs" Value="-deviceurl=https://horde.devtools.epicgames.com/ -reboot -cleardevices -sandbox=&quot;$(ProjectName)&quot;" If="$(IsBuildMachine)" />
