Fixing bug where base platforms were not being included as the list of TargetPlatforms was not empty.

#rnx

diff --git a/Engine/Build/Graph/Tasks/Inc/PlatformOptions.xml b/Engine/Build/Graph/Tasks/Inc/PlatformOptions.xml
--- a/Engine/Build/Graph/Tasks/Inc/PlatformOptions.xml
+++ b/Engine/Build/Graph/Tasks/Inc/PlatformOptions.xml
@@ -82,14 +82,12 @@
 	<Property Name="CookedEditorPlatforms" Value="$(CookedEditorPlatforms)Linux+" If="$(WithLinuxCookedEditor)" />
 </Do>
 
-<Do If="'$(TargetPlatforms)' == ''" >
-	<Property Name="TargetPlatforms" Value="$(TargetPlatforms)Android+" If="$(WithAndroid)" />
-	<Property Name="TargetPlatforms" Value="$(TargetPlatforms)IOS+" If="$(WithIOS)" />
-	<Property Name="TargetPlatforms" Value="$(TargetPlatforms)Linux+" If="$(WithLinux)" />
-	<Property Name="TargetPlatforms" Value="$(TargetPlatforms)Mac+" If="$(WithMac)" />
-	<Property Name="TargetPlatforms" Value="$(TargetPlatforms)TVOS+" If="$(WithTVOS)" />
-	<Property Name="TargetPlatforms" Value="$(TargetPlatforms)Win64+" If="$(WithWin64)" />
-</Do>
+<Expand Name="AddPlatform" Platform="Android" IsConsole="false" IsEnabled="$(WithAndroid)"/>
+<Expand Name="AddPlatform" Platform="IOS" IsConsole="false" IsEnabled="$(WithIOS)"/>
+<Expand Name="AddPlatform" Platform="Linux" IsConsole="false" IsEnabled="$(WithLinux)"/>
+<Expand Name="AddPlatform" Platform="Mac" IsConsole="false" IsEnabled="$(WithMac)"/>
+<Expand Name="AddPlatform" Platform="TVOS" IsConsole="false" IsEnabled="$(WithTVOS)"/>
+<Expand Name="AddPlatform" Platform="Win64" IsConsole="false" IsEnabled="$(WithWin64)"/>
 
 <!-- If still no platforms, pick the defaults that the script should have specified-->
 <Do If="'$(EditorPlatforms)' == ''" >
