Use Device Reservation Block annotations

#rnx

diff --git a/Engine/Build/Graph/Tasks/BuildAndTestProject.xml b/Engine/Build/Graph/Tasks/BuildAndTestProject.xml
--- a/Engine/Build/Graph/Tasks/BuildAndTestProject.xml
+++ b/Engine/Build/Graph/Tasks/BuildAndTestProject.xml
@@ -4,6 +4,8 @@
 	
 	<EnvVar Name="UE-CloudPublishHost"/>
 	<EnvVar Name="UE-CloudPublishNamespace"/>
+
+	<Option Name="Horde" Restrict="true|false" DefaultValue="$(IsBuildMachine)" Description="Whether the build is running under Horde"/>
 	
 	<!-- Project Options -->
 	<Option Name="ProjectName" DefaultValue="None" Restrict=".+" Description="Name of Project"/>
@@ -40,6 +42,7 @@
 	<Option Name="UnrealInsightsBuildConfiguration" DefaultValue="Shipping" Restrict="Development|Debug|Shipping" Description="In which configuration to build Unreal Insights, if included. Shipping, Debug, or Development." />
 	<Option Name="IncludeUnrealInsightsInToolCompile" DefaultValue="false" Description="Include Unreal Insights into Compile Tools prerequisites" />
 	<Option Name="IncludeEpicWebHelperInToolCompile" DefaultValue="false" Description="Include EpicWebHelper into Compile Tools prerequisites" />
+	<Option Name="WithDeviceReservation" DefaultValue="$(Horde)" Description="Whether to re-use devices and installed builds for a sequence of tests, where supported"/>
 
 	<!-- 
 		Allow project graphs that Include this file to specify their own defaults.
@@ -560,11 +563,14 @@
 		</Agent>
 
 		<!-- Test Time! -->
+		<Property Name="AllPlatformInstallNodes" Value="" />
 		<Property Name="AllPlatformBootTestNodes" Value="" />
 		<Property Name="AllPlatformTestNodes" Value="" />
 		<Agent Name="Test $(PlatformAgentName)" Type="$(TestAgentType)">
 			<!-- Run through all configs -->
 			<ForEach Name="TestConfiguration" Values="$(TestConfigurations)" Separator="+">
+				<Property Name="NodeAnnotations" Value=""/>
+
 				<!-- Prerequisites for the test to run -->
 				<Property Name="TestPrerequisites" Value="$(PlatformPackageNodeName)" />
 				
@@ -577,11 +583,34 @@
 						<Property Name="TestDirectory" Value="$(TestDirectory)/$(StagedPlatformFolder)/Packages/$(TestConfiguration)" />
 					</Do>
 					<Property Name="GauntletBuildArg" Value="-build=&quot;$(TestDirectory)&quot;" />
+					<Property Name="GauntletBuildArg" Value="$(GauntletBuildArg) -destlocalinstalldir=&quot;$(LocalInstallDirectory)&quot;" If="$(WithDeviceReservation)"/>
 				</Do>
 
 				<Property Name="GauntletTestArgs" Value="$(GauntletTargetTestArgs)"/>
 				<Property Name="GauntletTestArgs" Value="$(GauntletTargetTestAltArgs)" If="ContainsItem('$(AlternateDevicePoolPlatforms)', '$(TargetPlatform)', ';')"/>
 
+				<!-- Install step -->
+				<Do If="$(WithDeviceReservation)">
+					<Property Name="NodeAnnotations" Value=""/>
+					<Property Name="NodeAnnotations" Value="Workflow=$(AnnotationsTarget)" If="ContainsItem('$(AnnotationsPlatforms)', '$(TargetPlatform)', ';') and ContainsItem('$(AnnotationsTests)', 'Install', ';')"/>
+					<Property Name="InstallNodeName" Value="$(PreNodeName)$(TargetPlatform) $(TestConfiguration) Install Game"/>
+					<Property Name="InstallArgs" Value="-ClientCount=1 $(GauntletPreInstallArgs)"/>
+					<Do If="'$(AutomationReportOutputDirectory)' != ''">
+						<Property Name="OutputFragmentPath" Value="$(TargetPlatform)/$(TestConfiguration)/UE.InstallOnly" />
+						<Property Name="OutputReportPath" Value="$(AutomationReportOutputDirectory)/$(OutputFragmentPath)" />
+						<Property Name="GauntletReportArgs" Value="-logdir=&quot;$(OutputReportPath)&quot; -artifactname=Gauntlet" />
+					</Do>
+					<Node Name="$(InstallNodeName)" After="$(PlatformPublishPackagedNodeName)" Requires="$(TestPrerequisites)" Annotations="DeviceReserve=Begin;$(NodeAnnotations)">
+						<Command Name="RunUnreal" Arguments="-project=$(TargetProject) -test=&quot;UE.InstallOnly&quot; -platform=$(TargetPlatform) -configuration=$(TestConfiguration) -target=$(TargetName) $(GauntletBuildArg) $(GauntletTestArgs) $(GauntletReportArgs) $(InstallArgs)" />
+					</Node>
+					<Property Name="TestPrerequisites" Value="$(TestPrerequisites);$(InstallNodeName)" />
+					<Do If="!ContainsItem('$(TargetPlatformsToSkipBootTests)', '$(TargetPlatform)', '+') Or !ContainsItem('$(TargetPlatformsToSkipTests)', '$(TargetPlatform)', '+')">
+						<Property Name="AllPlatformInstallNodes" Value="$(AllPlatformInstallNodes);$(InstallNodeName)"/>
+					</Do>
+					<!-- Skip install for all other test node -->
+					<Property Name="GauntletTestArgs" Value="$(GauntletTestArgs) $(GauntletTestSkipInstallArgs)"/>
+				</Do>
+
 				<!-- Declare a boot test for the target-->
 				<Property Name="NodeAnnotations" Value=""/>
 				<Property Name="NodeAnnotations" Value="Workflow=$(AnnotationsTarget)" If="ContainsItem('$(AnnotationsPlatforms)', '$(TargetPlatform)', ';') and ContainsItem('$(AnnotationsTests)', 'BootTest', ';')"/>
@@ -590,8 +619,7 @@
 					<Do If="'$(AutomationReportOutputDirectory)' != ''">
 						<Property Name="OutputFragmentPath" Value="$(TargetPlatform)/$(TestConfiguration)/UE.BootTest" />
 						<Property Name="OutputReportPath" Value="$(AutomationReportOutputDirectory)/$(OutputFragmentPath)" />
-						<Property Name="GauntletReportArgs" Value="-logdir=&quot;$(OutputReportPath)&quot; -artifactname=Gauntlet -ReportExportPath=&quot;$(OutputReportPath)&quot;" />
-						<Property Name="GauntletReportArgs" Value="$(GauntletReportArgs) -ReportURL=&quot;$(AutomationReportUrl)/$(OutputFragmentPath)/&quot;" If="'$(AutomationReportUrl)' != ''" />
+						<Property Name="GauntletReportArgs" Value="-logdir=&quot;$(OutputReportPath)&quot; -artifactname=Gauntlet" />
 					</Do>
 					<!-- Run the boot test -->
 					<Command Name="RunUnreal" Arguments="-project=$(TargetProject) -test=&quot;UE.BootTest&quot; -platform=$(TargetPlatform) -configuration=$(TestConfiguration) -target=$(TargetName) $(GauntletBuildArg) $(GauntletTestArgs) $(GauntletReportArgs)" />
@@ -632,6 +660,7 @@
 		<Label Category="Clients" Name="$(PreLabelName)$(TargetPlatform)" Requires="$(PlatformCompileNodeName);$(UploadSymbolsNodeName)" />
 		<Label Category="Clients" Name="$(PreLabelName)Content $(TargetPlatform)" Requires="$(PlatformCookNodeName);$(PlatformStageNodeName);$(PlatformPackageNodeName)" />
 		<Label Category="Clients" Name="$(PreLabelName)Publish $(TargetPlatform)" Requires="$(PlatformPublishPackagedNodeName);$(PlatformPublishPackagedNodeName);$(PlatformPublishStagedNodeName);$(PlatformHashNodeName)" />
+		<Label Category="Tests" Name="$(PreLabelName)Install $(TargetPlatform)" Requires="$(AllPlatformInstallNodes)" />
 		<Label Category="Tests" Name="$(PreLabelName)BootTest $(TargetPlatform)" Requires="$(AllPlatformBootTestNodes)" />
 		<Label Category="Tests" Name="$(PreLabelName)Test $(TargetPlatform)" Requires="$(AllPlatformTestNodes)" />
 
@@ -641,7 +670,7 @@
 		<Property Name="PlatformTasks" Value="$(PlatformTasks);$(PlatformCookNodeName)" If="!$(SkipCook) and !$(SkipPublish)" />		
 		<Property Name="PlatformTasks" Value="$(PlatformTasks);$(PlatformPackageNodeName);$(PlatformPublishPackagedNodeName)" If="!$(SkipPublish) And $(PublishBuilds)" />	
 		<Property Name="PlatformTasks" Value="$(PlatformTasks);$(PlatformHashNodeName)"  If="$(PublishHashedStagingData)" />	
-		<Property Name="PlatformTasks" Value="$(PlatformTasks);$(AllPlatformBootTestNodes);$(AllPlatformTestNodes)" />		
+		<Property Name="PlatformTasks" Value="$(PlatformTasks);$(AllPlatformBootTestNodes);$(AllPlatformTestNodes);$(AllPlatformInstallNodes)" />		
 
 		<Property Name="GeneratedTargetNodes" Value="$(GeneratedTargetNodes);$(PlatformTasks)"/>
 	</ForEach>
diff --git a/Engine/Build/Graph/Tasks/Inc/GauntletSettings.xml b/Engine/Build/Graph/Tasks/Inc/GauntletSettings.xml
--- a/Engine/Build/Graph/Tasks/Inc/GauntletSettings.xml
+++ b/Engine/Build/Graph/Tasks/Inc/GauntletSettings.xml
@@ -41,6 +41,10 @@
 	<Property Name="GauntletEditorTestArgs" Value="-build=editor $(GauntletCommonArgs) $(ExtraGauntletEditorTestArgs)" />
 	<Property Name="GauntletCookedEditorTestArgs" Value="-cookededitor $(GauntletCommonArgs) $(ExtraGauntletEditorTestArgs)" />
 
+	<!-- Arguments for handling devices when using device reservation block -->
+	<Property Name="GauntletPreInstallArgs" Value="-fullclean"/>
+	<Property Name="GauntletTestSkipInstallArgs" Value="-skipinstall"/>
+
     <!-- this is fully generated dynamiclly when needed, it's simply here so it can be used without worrying if its set -->
     <Property Name="GauntletReportArgs" Value="" />
 
diff --git a/Engine/Build/Graph/Tasks/Inc/PathSettings.xml b/Engine/Build/Graph/Tasks/Inc/PathSettings.xml
--- a/Engine/Build/Graph/Tasks/Inc/PathSettings.xml
+++ b/Engine/Build/Graph/Tasks/Inc/PathSettings.xml
@@ -41,6 +41,7 @@
 	<!-- Where builds are saved to locally -->
 	<Property Name="ProjectOutputDirectory" Value="$(ProjectPath)/Saved/StagedBuilds"/>
 	<Property Name="ShaderSymbolsDirectory" Value="$(ProjectPath)/Saved/ShaderSymbols"/>
+	<Property Name="LocalInstallDirectory" Value="$(ProjectPath)/Saved/LocalInstall"/>
 	<Property Name="AutomationReportOutputDirectory" Value=""/>
 
 	<!-- Scratch network path that can be used for tests if we aren't publishing. Assumes there's a cleanup script.. -->
