Exclude target platforms that do not need reservation

#rnx

diff --git a/Engine/Build/Graph/Tasks/BuildAndTestProject.xml b/Engine/Build/Graph/Tasks/BuildAndTestProject.xml
--- a/Engine/Build/Graph/Tasks/BuildAndTestProject.xml
+++ b/Engine/Build/Graph/Tasks/BuildAndTestProject.xml
@@ -405,15 +405,19 @@
 	<!-- Now, for each platform we build, cook, stage, package, publish, and test -->
 	<ForEach Name="TargetPlatform" Values="$(TargetPlatforms)" Separator="+">
 
+		<Property Name="IsDesktopPlatform" Value="ContainsItem('$(AllDesktopPlatforms)', '$(TargetPlatform)', '+')"/>
+		<Property Name="IsMobilePlatform" Value="'$(TargetPlatform)' == 'Android' or '$(TargetPlatform)' == 'IOS'"/>
+		<!-- Local or virtualized platform -->
+		<Property Name="IsVirtualizedPlatform" Value="ContainsItem('$(AllVirtualizedPlatforms)', '$(TargetPlatform)', '+')"/>
+		<Property Name="IsLocalPlatform" Value="($(IsDesktopPlatform) Or $(IsVirtualizedPlatform))"/>
+
 		<!-- Declare type of host required to build & cook this platform-->
 		<Property Name="HostAgentType" Value="Win64" />
 		<Property Name="HostAgentType" Value="Mac" If="'$(TargetPlatform)' == 'Mac' or '$(TargetPlatform)' == 'IOS' or '$(TargetPlatform)' == 'tvOS'"/>
 		<Property Name="HostAgentType" Value="Linux" If="'$(TargetPlatform)' == 'Linux'"/>
 		<!-- On consoles we can use agents with devkit setup, but locally hosted tests need a Test agent-->
 		<Property Name="TestAgentType" Value="$(TargetPlatformTestAgentOverride)" />
-		<Property Name="TestAgentType" Value="Test$(HostAgentType)" If="'$(TargetPlatform)' == 'Win64' or '$(TargetPlatform)' == 'Mac' or '$(TargetPlatform)' == 'Linux'"/>
-
-		<Property Name="IsMobilePlatform" Value="'$(TargetPlatform)' == 'Android' or '$(TargetPlatform)' == 'IOS'"/>
+		<Property Name="TestAgentType" Value="Test$(HostAgentType)" If="$(IsDesktopPlatform)"/>
 
 		<!-- Use incremental agents for building? -->
 		<Property Name="CompileAgentType" Value="$(HostAgentType)" />
@@ -583,14 +587,14 @@
 						<Property Name="TestDirectory" Value="$(TestDirectory)/$(StagedPlatformFolder)/Packages/$(TestConfiguration)" />
 					</Do>
 					<Property Name="GauntletBuildArg" Value="-build=&quot;$(TestDirectory)&quot;" />
-					<Property Name="GauntletBuildArg" Value="$(GauntletBuildArg) -destlocalinstalldir=&quot;$(LocalInstallDirectory)&quot;" If="$(WithDeviceReservation)"/>
+					<Property Name="GauntletBuildArg" Value="$(GauntletBuildArg) -destlocalinstalldir=&quot;$(LocalInstallDirectory)&quot;" If="!$(IsLocalPlatform) And $(WithDeviceReservation)"/>
 				</Do>
 
 				<Property Name="GauntletTestArgs" Value="$(GauntletTargetTestArgs)"/>
 				<Property Name="GauntletTestArgs" Value="$(GauntletTargetTestAltArgs)" If="ContainsItem('$(AlternateDevicePoolPlatforms)', '$(TargetPlatform)', ';')"/>
 
 				<!-- Install step -->
-				<Do If="$(WithDeviceReservation)">
+				<Do If="!$(IsLocalPlatform) And $(WithDeviceReservation)">
 					<Property Name="NodeAnnotations" Value=""/>
 					<Property Name="NodeAnnotations" Value="Workflow=$(AnnotationsTarget)" If="ContainsItem('$(AnnotationsPlatforms)', '$(TargetPlatform)', ';') and ContainsItem('$(AnnotationsTests)', 'Install', ';')"/>
 					<Property Name="InstallNodeName" Value="$(PreNodeName)$(TargetPlatform) $(TestConfiguration) Install Game"/>
diff --git a/Engine/Build/Graph/Tasks/Inc/PlatformOptions.xml b/Engine/Build/Graph/Tasks/Inc/PlatformOptions.xml
--- a/Engine/Build/Graph/Tasks/Inc/PlatformOptions.xml
+++ b/Engine/Build/Graph/Tasks/Inc/PlatformOptions.xml
@@ -24,9 +24,11 @@
 		OutProperty: RequiredEditorPlatforms - Editor platforms to build to cook targets      
 	-->
 
-<Property Name="AllEditorPlatforms" Value="Mac+Win64+Linux" />
+<Property Name="AllDesktopPlatforms" Value="Mac+Win64+Linux" />
+<Property Name="AllEditorPlatforms" Value="$(AllDesktopPlatforms)" />
 <Property Name="AllConsolePlatforms" Value="" />
 <Property Name="AllTargetPlatforms" Value="Android+IOS+Mac+TVOS+Win64+Linux" />
+<Property Name="AllVirtualizedPlatforms" Value="" />
 
 <Option Name="EditorPlatforms" DefaultValue="" Description="Alternative to 'WithFooEditor', a list of editors to build and test. Win64+Mac etc"/>
 <Option Name="TargetPlatforms" DefaultValue="" Description="Alternative to 'WithFoo', a list of Platforms to build and test. Android+IOS etc"/>
@@ -44,9 +46,11 @@
 <Option Name="CookedEditorPlatformsToSkipTests" DefaultValue="" Restrict=".*" Description="List of editor cooked platforms that editor tests won't be run on" />
 <Option Name="PlatformsToSkipTests" DefaultValue="" Restrict=".*" Description="List of platforms that all tests won't be run on. Overrides Editor/TargetPlatformsToSkipTests when non-empty." />
 
-<Macro Name="AddPlatform" Arguments="Platform;IsConsole;IsEnabled">
+<Macro Name="AddPlatform" Arguments="Platform;IsConsole;IsEnabled" OptionalArguments="IsVirtualized">
+	<Property Name="IsVirtualized" Value="('$(IsVirtualized)' != '' and '$(IsVirtualized)' != 'false')"/>
 	<Property Name="AllTargetPlatforms" Value="$(AllTargetPlatforms)+$(Platform)" If="!ContainsItem('$(AllTargetPlatforms)', '$(Platform)', '+')"/>
 	<Property Name="AllConsolePlatforms" Value="$(AllConsolePlatforms)+$(Platform)" If="$(IsConsole) and !ContainsItem('$(AllConsolePlatforms)', '$(Platform)', '+')"/>
+	<Property Name="AllVirtualizedPlatforms" Value="$(AllVirtualizedPlatforms)+$(Platform)" If="$(IsVirtualized) And !ContainsItem('$(AllVirtualizedPlatforms)', '$(Platform)', '+')"/>
 	<Property Name="TargetPlatforms" Value="$(TargetPlatforms)+$(Platform)" If="$(IsEnabled) and !ContainsItem('$(TargetPlatforms)', '$(Platform)', '+')"/>
 </Macro>
 	
