Add an example script that shows using different parameter types.

diff --git a/Engine/Build/Graph/Examples/Parameters.xml b/Engine/Build/Graph/Examples/Parameters.xml
new file mode 100644
--- /dev/null
+++ b/Engine/Build/Graph/Examples/Parameters.xml
@@ -0,0 +1,47 @@
+<?xml version='1.0' ?>
+<BuildGraph xmlns="http://www.epicgames.com/BuildGraph" xmlns:xsi="http://www.w3.org/2001/XMLSchema-instance" xsi:schemaLocation="http://www.epicgames.com/BuildGraph ../Schema.xsd" >
+
+	<!-- Text parameter -->
+	<Option Name="TextParameter" DefaultValue=""/>
+	
+	<!-- List parameter -->
+	<Option Name="ListParameter" Restrict="Option1|Option2|Option3" DefaultValue="Option2"/>
+
+	<!-- Multi-list parameter: Options are specified as separate booleans from which we can reconstruct a list -->
+	<Option Name="MultiListOption1" Restrict="true|false" DefaultValue="false"/>
+	<Option Name="MultiListOption2" Restrict="true|false" DefaultValue="false"/>
+	<Option Name="MultiListOption3" Restrict="true|false" DefaultValue="false"/>
+	<Option Name="MultiListOption4" Restrict="true|false" DefaultValue="false"/>
+
+	<Property Name="MultiListParameters" Value=""/>
+	<Property Name="MultiListParameters" Value="$(MultiListParameters)MultiList1;" If="$(MultiListOption1)"/>
+	<Property Name="MultiListParameters" Value="$(MultiListParameters)MultiList2;" If="$(MultiListOption2)"/>
+	<Property Name="MultiListParameters" Value="$(MultiListParameters)MultiList3;" If="$(MultiListOption3)"/>
+	<Property Name="MultiListParameters" Value="$(MultiListParameters)MultiList4;" If="$(MultiListOption4)"/>
+
+	<!-- Tag-picker parameter: Similar to multi-list parameter -->
+	<Option Name="TagPickerOption1" Restrict="true|false" DefaultValue="false"/>
+	<Option Name="TagPickerOption2" Restrict="true|false" DefaultValue="false"/>
+	<Option Name="TagPickerOption3" Restrict="true|false" DefaultValue="false"/>
+	<Option Name="TagPickerOption4" Restrict="true|false" DefaultValue="false"/>
+
+	<Property Name="TagPickerParameters" Value=""/>
+	<Property Name="TagPickerParameters" Value="$(TagPickerParameters)TagPicker1;" If="$(TagPickerOption1)"/>
+	<Property Name="TagPickerParameters" Value="$(TagPickerParameters)TagPicker2;" If="$(TagPickerOption2)"/>
+	<Property Name="TagPickerParameters" Value="$(TagPickerParameters)TagPicker3;" If="$(TagPickerOption3)"/>
+	<Property Name="TagPickerParameters" Value="$(TagPickerParameters)TagPicker4;" If="$(TagPickerOption4)"/>
+
+	<!-- Bool parameter -->
+	<Option Name="BoolParameter" Restrict="true|false" DefaultValue="false"/>
+
+	<Agent Name="Test Agent" Type="Win64">
+		<Node Name="Print Parameters">
+			<Log Message="Text Parameter: $(TextParameter)"/>
+			<Log Message="List Parameter: $(TextParameter)"/>
+			<Log Message="MultiList Parameter: $(MultiListParameters)"/>
+			<Log Message="TagPicker Parameter: $(TagPickerParameter)"/>
+			<Log Message="Bool Parameter: $(BoolParameter)"/>
+		</Node>
+	</Agent>
+	
+</BuildGraph>
