Add missing description attributes to Parameters.xml script.

diff --git a/Engine/Build/Graph/Examples/Parameters.xml b/Engine/Build/Graph/Examples/Parameters.xml
--- a/Engine/Build/Graph/Examples/Parameters.xml
+++ b/Engine/Build/Graph/Examples/Parameters.xml
@@ -2,16 +2,16 @@
 <BuildGraph xmlns="http://www.epicgames.com/BuildGraph" xmlns:xsi="http://www.w3.org/2001/XMLSchema-instance" xsi:schemaLocation="http://www.epicgames.com/BuildGraph ../Schema.xsd" >
 
 	<!-- Text parameter -->
-	<Option Name="TextParameter" DefaultValue=""/>
+	<Option Name="TextParameter" DefaultValue="" Description="Text parameter"/>
 	
 	<!-- List parameter -->
-	<Option Name="ListParameter" Restrict="Option1|Option2|Option3" DefaultValue="Option2"/>
+	<Option Name="ListParameter" Restrict="Option1|Option2|Option3" DefaultValue="Option2" Description="List parameter"/>
 
 	<!-- Multi-list parameter: Options are specified as separate booleans from which we can reconstruct a list -->
-	<Option Name="MultiListOption1" Restrict="true|false" DefaultValue="false"/>
-	<Option Name="MultiListOption2" Restrict="true|false" DefaultValue="false"/>
-	<Option Name="MultiListOption3" Restrict="true|false" DefaultValue="false"/>
-	<Option Name="MultiListOption4" Restrict="true|false" DefaultValue="false"/>
+	<Option Name="MultiListOption1" Restrict="true|false" DefaultValue="false" Description="MultiList Parameter Option 1"/>
+	<Option Name="MultiListOption2" Restrict="true|false" DefaultValue="false" Description="MultiList Parameter Option 2"/>
+	<Option Name="MultiListOption3" Restrict="true|false" DefaultValue="false" Description="MultiList Parameter Option 3"/>
+	<Option Name="MultiListOption4" Restrict="true|false" DefaultValue="false" Description="MultiList Parameter Option 4"/>
 
 	<Property Name="MultiListParameters" Value=""/>
 	<Property Name="MultiListParameters" Value="$(MultiListParameters)MultiList1;" If="$(MultiListOption1)"/>
@@ -20,10 +20,10 @@
 	<Property Name="MultiListParameters" Value="$(MultiListParameters)MultiList4;" If="$(MultiListOption4)"/>
 
 	<!-- Tag-picker parameter: Similar to multi-list parameter -->
-	<Option Name="TagPickerOption1" Restrict="true|false" DefaultValue="false"/>
-	<Option Name="TagPickerOption2" Restrict="true|false" DefaultValue="false"/>
-	<Option Name="TagPickerOption3" Restrict="true|false" DefaultValue="false"/>
-	<Option Name="TagPickerOption4" Restrict="true|false" DefaultValue="false"/>
+	<Option Name="TagPickerOption1" Restrict="true|false" DefaultValue="false" Description="TagPicker Parameter Option 1"/>
+	<Option Name="TagPickerOption2" Restrict="true|false" DefaultValue="false" Description="TagPicker Parameter Option 2"/>
+	<Option Name="TagPickerOption3" Restrict="true|false" DefaultValue="false" Description="TagPicker Parameter Option 3"/>
+	<Option Name="TagPickerOption4" Restrict="true|false" DefaultValue="false" Description="TagPicker Parameter Option 4"/>
 
 	<Property Name="TagPickerParameters" Value=""/>
 	<Property Name="TagPickerParameters" Value="$(TagPickerParameters)TagPicker1;" If="$(TagPickerOption1)"/>
