Horde: Add missing parameter to build script.

diff --git a/Engine/Build/Graph/Examples/Parameters.xml b/Engine/Build/Graph/Examples/Parameters.xml
--- a/Engine/Build/Graph/Examples/Parameters.xml
+++ b/Engine/Build/Graph/Examples/Parameters.xml
@@ -32,7 +32,7 @@
 	<Property Name="TagPickerParameters" Value="$(TagPickerParameters)TagPicker4;" If="$(TagPickerOption4)"/>
 
 	<!-- Bool parameter -->
-	<Option Name="BoolParameter" Restrict="true|false" DefaultValue="false"/>
+	<Option Name="BoolParameter" Restrict="true|false" DefaultValue="false" Description="Bool parameter option"/>
 
 	<Agent Name="Test Agent" Type="Win64">
 		<Node Name="Print Parameters">
