Horde: Fix incorrect property name in build script.

diff --git a/Engine/Build/Graph/Examples/Parameters.xml b/Engine/Build/Graph/Examples/Parameters.xml
--- a/Engine/Build/Graph/Examples/Parameters.xml
+++ b/Engine/Build/Graph/Examples/Parameters.xml
@@ -39,7 +39,7 @@
 			<Log Message="Text Parameter: $(TextParameter)"/>
 			<Log Message="List Parameter: $(TextParameter)"/>
 			<Log Message="MultiList Parameter: $(MultiListParameters)"/>
-			<Log Message="TagPicker Parameter: $(TagPickerParameter)"/>
+			<Log Message="TagPicker Parameter: $(TagPickerParameters)"/>
 			<Log Message="Bool Parameter: $(BoolParameter)"/>
 		</Node>
 	</Agent>
