Horde: Add labels to BuildEditorAndTools script.

diff --git a/Engine/Build/Graph/Examples/BuildEditorAndTools.xml b/Engine/Build/Graph/Examples/BuildEditorAndTools.xml
--- a/Engine/Build/Graph/Examples/BuildEditorAndTools.xml
+++ b/Engine/Build/Graph/Examples/BuildEditorAndTools.xml
@@ -179,4 +179,16 @@
 		
 	</Agent>
 
+	<Label Category="Editors" Name="Win64" Requires="Compile $(EditorTarget) Win64" UgsBadge="Editor"/>
+
+	<ForEach Name="TargetPlatform" Values="$(TargetPlatforms)">
+		<Property Name="GameTargetNodes" Value=""/>	
+		<ForEach Name="GameTarget" Values="$(GameTargets)">
+			<Property Name="GameTargetNodes" Value="$(GameTargetNodes)Compile $(GameTarget) $(TargetPlatform);"/>
+		</ForEach>
+		<Do If="'$(GameTargetNodes)' != ''">
+			<Label Category="Target Platforms" Name="$(TargetPlatform)" Requires="$(GameTargetNodes)" UgsBadge="$(TargetPlatform)"/>
+		</Do>
+	</ForEach>
+
 </BuildGraph>
