Horde: Fix build error in default configuration of packaged build script.

* BuildAndTestProject.xml now takes optional -set:ProjectFile=... argument rather than specifying name and directory individually.
* Also improved error message for invalid object names in buildgraph script parser.

diff --git a/Engine/Build/Graph/Tasks/BuildAndTestProject.xml b/Engine/Build/Graph/Tasks/BuildAndTestProject.xml
--- a/Engine/Build/Graph/Tasks/BuildAndTestProject.xml
+++ b/Engine/Build/Graph/Tasks/BuildAndTestProject.xml
@@ -8,8 +8,14 @@
 	<Option Name="Horde" Restrict="true|false" DefaultValue="$(IsBuildMachine)" Description="Whether the build is running under Horde"/>
 	
 	<!-- Project Options -->
-	<Option Name="ProjectName" DefaultValue="None" Restrict=".+" Description="Name of Project"/>
-	<Option Name="ProjectPath" DefaultValue="None" Restrict=".+" Description="Path to the folder that contains your project"/>
+	<Option Name="ProjectFile" DefaultValue="" Restrict=".*" Description="Path to the project file. Can be used instead of specifying ProjectName and ProjectPath separately."/>
+
+	<Property Name="DefaultProjectName" Value="None"/>
+	<Property Name="DefaultProjectPath" Value="None"/>
+	<Regex Input="$(ProjectFile)" Pattern="(.*)[/\\]([^/\\]+)\.uproject" Capture="DefaultProjectPath;DefaultProjectName" If="'$(ProjectFile)' != ''"/>
+
+	<Option Name="ProjectName" DefaultValue="$(DefaultProjectName)" Restrict=".+" Description="Name of Project"/>
+	<Option Name="ProjectPath" DefaultValue="$(DefaultProjectPath)" Restrict=".+" Description="Path to the folder that contains your project"/>
 
 	<Option Name="IsContentOnlyProject" DefaultValue="false" Restrict="true|false" Description="Is the project a content only project?"/>	
 
diff --git a/Engine/Source/Programs/AutomationTool/BuildGraph/BgScriptReader.cs b/Engine/Source/Programs/AutomationTool/BuildGraph/BgScriptReader.cs
--- a/Engine/Source/Programs/AutomationTool/BuildGraph/BgScriptReader.cs
+++ b/Engine/Source/Programs/AutomationTool/BuildGraph/BgScriptReader.cs
@@ -1952,12 +1952,12 @@ namespace AutomationTool
 			{
 				if (idx > 0 && name[idx] == ' ' && name[idx - 1] == ' ')
 				{
-					LogError(element, "Consecutive spaces in object name - '{Name}'", name);
+					LogError(element, "Consecutive spaces in object name '{Name}'", name);
 					return false;
 				}
 				if (Char.IsControl(name[idx]) || BgScriptSchema.IllegalNameCharacters.IndexOf(name[idx]) != -1)
 				{
-					LogError(element, "Invalid character in object name - '{Name}'", name[idx]);
+					LogError(element, "Invalid character in object name '{Name}': '{Character}'", name, name[idx]);
 					return false;
 				}
 			}
diff --git a/Engine/Source/Programs/Horde/Horde.Server/Defaults/ue5-dev-main-hordedemo.stream.json b/Engine/Source/Programs/Horde/Horde.Server/Defaults/ue5-dev-main-hordedemo.stream.json
--- a/Engine/Source/Programs/Horde/Horde.Server/Defaults/ue5-dev-main-hordedemo.stream.json
+++ b/Engine/Source/Programs/Horde/Horde.Server/Defaults/ue5-dev-main-hordedemo.stream.json
@@ -192,8 +192,7 @@
 			"arguments": [
 				"-Script=Engine/Build/Graph/Tasks/BuildAndTestProject.xml",
 				"-Target=BuildAndTest $(Project)",
-				"-set:ProjectName=$(Project)",
-				"-set:ProjectPath=$(ProjectPath)",
+				"-set:ProjectFile=$(ProjectPath)",
 				"-set:TargetName=$(GameTarget)",
 				"-set:EditorTargetName=$(EditorTarget)"
 			],
