Fix issue with Test nodes Packaged dependency

#rnx

diff --git a/Engine/Build/Graph/Tasks/BuildAndTestProject.xml b/Engine/Build/Graph/Tasks/BuildAndTestProject.xml
--- a/Engine/Build/Graph/Tasks/BuildAndTestProject.xml
+++ b/Engine/Build/Graph/Tasks/BuildAndTestProject.xml
@@ -452,6 +452,10 @@
 		<Property Name="PlatformHashNodeName" Value="$(PreNodeName)Hash $(TargetPlatform)" />
 		<Property Name="PlatformPublishStagedNodeName" Value="$(PreNodeName)Publish Staged $(TargetPlatform)" />
 		<Property Name="PlatformPublishPackagedNodeName" Value="$(PreNodeName)Publish Packaged $(TargetPlatform)" />
+		<Property Name="RequirePlatformStaged" Value="$(PlatformStageNodeName)" />
+		<Property Name="RequirePlatformStaged" Value="$(PlatformPublishStagedNodeName)" If="!$(SkipPublish) And $(PublishBuilds)"/>
+		<Property Name="RequirePlatformPackaged" Value="$(PlatformPackageNodeName)" />
+		<Property Name="RequirePlatformPackaged" Value="$(PlatformPublishPackagedNodeName)"  If="!$(SkipPublish) And $(PublishBuilds)"/>
 		
 		<!-- Set CookPlatform since it isn't always the same as the target -->
 		<Property Name="CookPlatform" Value="$(TargetPlatform)" />
@@ -516,9 +520,11 @@
 
 			<!-- Cook -->
 			<Node Name="$(PlatformCookNodeName)" Requires="$(PlatformEditorNodeName)" Produces="#$(PlatformCookNodeName) Complete" RunEarly="true">
-				<Property Name="CookArgs" Value="" />
-				<Property Name="CookArgs" Value="-MapIniSection=$(MapIniSection) $(CookPlatform)" If="'$(MapIniSection)' != ''"/>				
-				<Cook Project="$(TargetProject)" Platform="$(CookPlatform)" TagOutput="false" Arguments="$(CookArgs) $(ExtraCookArguments)"/>
+				<Do If="!$(SkipCook)">
+					<Property Name="CookArgs" Value="" />
+					<Property Name="CookArgs" Value="-MapIniSection=$(MapIniSection) $(CookPlatform)" If="'$(MapIniSection)' != ''"/>				
+					<Cook Project="$(TargetProject)" Platform="$(CookPlatform)" TagOutput="false" Arguments="$(CookArgs) $(ExtraCookArguments)"/>
+				</Do>
 			</Node>
 			
 			<!-- Stage -->
@@ -559,7 +565,7 @@
 			</Node>
 			
 			<!-- Package -->
-			<Node Name="$(PlatformPackageNodeName)" Requires="$(PlatformPublishStagedNodeName)">
+			<Node Name="$(PlatformPackageNodeName)" Requires="$(PlatformPublishStageDependencyName)">
 				<ForEach Name="PackageConfiguration" Values="$(PackageConfigurations)" Separator="+">
 					<Command Name="BuildCookRun" Arguments="$(BCRArgs) -configuration=$(PackageConfiguration) $(StageAndPackageArguments) -skipbuild -skipcook -skipstage -skippak -package $(PackageOnlyArguments)" />
 				</ForEach>
@@ -584,7 +590,8 @@
 				<Property Name="NodeAnnotations" Value=""/>
 
 				<!-- Prerequisites for the test to run -->
-				<Property Name="TestPrerequisites" Value="$(PlatformPackageNodeName)" />
+				<Property Name="TestPrerequisites" Value="$(RequirePlatformStaged)" />
+				<Property Name="TestPrerequisites" Value="$(RequirePlatformPackaged)" If="ContainsItem('$(PlatformsThatRequirePackaging)', '$(TargetPlatform)', '+') Or !ContainsItem('$(PlatformsToSkipPackaging)', '$(TargetPlatform)', '+')"/>
 				<Property Name="TestPrerequisites" Value="$(PlatformToolsNodeName)" If="'$(GeneratedBuildVersion)' != ''"/>
 				<Property Name="TestPrerequisites" Value="$(PlatformToolsNodeName);$(PlatformEditorNodeName)" If="'$(GeneratedBuildVersion)' != '' and $(TestsRequireEditor)"/>
 				
@@ -617,7 +624,7 @@
 						<Property Name="OutputReportPath" Value="$(AutomationReportOutputDirectory)/$(OutputFragmentPath)" />
 						<Property Name="GauntletReportArgs" Value="-logdir=&quot;$(OutputReportPath)&quot; -artifactname=Gauntlet" />
 					</Do>
-					<Node Name="$(InstallNodeName)" After="$(PlatformPublishPackagedNodeName)" Requires="$(TestPrerequisites)" Annotations="DeviceReserve=Begin;$(NodeAnnotations)">
+					<Node Name="$(InstallNodeName)" Requires="$(TestPrerequisites)" Annotations="DeviceReserve=Begin;$(NodeAnnotations)">
 						<Command Name="RunUnreal" Arguments="-project=$(TargetProject) -test=&quot;UE.InstallOnly&quot; -platform=$(TargetPlatform) -configuration=$(TestConfiguration) -target=$(TargetName) $(GauntletBuildArg) $(GauntletTestArgs) $(GauntletReportArgs) $(InstallArgs)" />
 					</Node>
 					<Property Name="TestPrerequisites" Value="$(TestPrerequisites);$(InstallNodeName)" />
@@ -632,7 +639,7 @@
 				<Property Name="NodeAnnotations" Value=""/>
 				<Property Name="NodeAnnotations" Value="Workflow=$(AnnotationsTarget)" If="ContainsItem('$(AnnotationsPlatforms)', '$(TargetPlatform)', ';') and ContainsItem('$(AnnotationsTests)', 'BootTest', ';')"/>
 				<Property Name="BootTestName" Value="$(PreNodeName)$(TargetPlatform) $(TestConfiguration) BootTest" />
-				<Node Name="$(BootTestName)" After="$(PlatformPublishPackagedNodeName)" Requires="$(TestPrerequisites)" Annotations="$(NodeAnnotations)">
+				<Node Name="$(BootTestName)" Requires="$(TestPrerequisites)" Annotations="$(NodeAnnotations)">
 					<Do If="'$(AutomationReportOutputDirectory)' != ''">
 						<Property Name="OutputFragmentPath" Value="$(TargetPlatform)/$(TestConfiguration)/UE.BootTest" />
 						<Property Name="OutputReportPath" Value="$(AutomationReportOutputDirectory)/$(OutputFragmentPath)" />
@@ -684,8 +691,9 @@
 		<!-- Tie all steps into our platform tasks unless excluded -->
 		<Property Name="PlatformTasks" Value="" />
 		<Property Name="PlatformTasks" Value="$(PlatformCompileNodeName);$(UploadSymbolsNodeName)" If="!$(SkipBuild)" />
-		<Property Name="PlatformTasks" Value="$(PlatformTasks);$(PlatformCookNodeName)" If="!$(SkipCook) and !$(SkipPublish)" />		
-		<Property Name="PlatformTasks" Value="$(PlatformTasks);$(PlatformPackageNodeName);$(PlatformPublishPackagedNodeName)" If="!$(SkipPublish) And $(PublishBuilds)" />	
+		<Property Name="PlatformTasks" Value="$(PlatformTasks);$(PlatformCookNodeName)" If="!$(SkipCook) and !$(SkipPublish)" />
+		<Property Name="PlatformTasks" Value="$(PlatformTasks);$(RequirePlatformStaged)" If="!$(SkipPublish)" />
+		<Property Name="PlatformTasks" Value="$(PlatformTasks);$(RequirePlatformPackaged)" If="!$(SkipPackage) and !$(SkipPublish)" />
 		<Property Name="PlatformTasks" Value="$(PlatformTasks);$(PlatformHashNodeName)"  If="$(PublishHashedStagingData)" />	
 		<Property Name="PlatformTasks" Value="$(PlatformTasks);$(AllPlatformBootTestNodes);$(AllPlatformTestNodes)" />		
 
diff --git a/Engine/Build/Graph/Tasks/Inc/CommonProperties.xml b/Engine/Build/Graph/Tasks/Inc/CommonProperties.xml
--- a/Engine/Build/Graph/Tasks/Inc/CommonProperties.xml
+++ b/Engine/Build/Graph/Tasks/Inc/CommonProperties.xml
@@ -16,7 +16,7 @@
 	<Property Name="IsPreflight" Value="false"/>
 	<Property Name="IsPreflight" Value="true" If="'$(PreflightChange)' != '' And $(SetPreflightFlag)"/>
 	<Property Name="PreflightSuffix" Value=""/>
-	<Property Name="PreflightSuffix" Value="-PF-$(PreflightChange)-$(UE_HORDE_JOBID)" If="$(IsPreflight)"/>
+	<Property Name="PreflightSuffix" Value="-PF-$(PreflightChange)-$(UE_HORDE_JOBID:-6)" If="$(IsPreflight)"/>
 	<Property Name="CL-String" Value="CL-$(Change)$(PreflightSuffix)"/>
 	<Property Name="BuildNamePath" Value="$(ProjectName)-$(CL-String)"/>
 	<Property Name="BuildVersion" Value="$(ProjectName)-CL-$(Change)"/>
diff --git a/Engine/Build/Graph/Tasks/Inc/PlatformOptions.xml b/Engine/Build/Graph/Tasks/Inc/PlatformOptions.xml
--- a/Engine/Build/Graph/Tasks/Inc/PlatformOptions.xml
+++ b/Engine/Build/Graph/Tasks/Inc/PlatformOptions.xml
@@ -156,7 +156,7 @@
 </ForEach>
 
 <!-- Packaging isn't supported by Gauntlet yet :( )-->
-<Property Name="PlatformsToSkipPackaging" Value="Win64+Linux" />
+<Property Name="PlatformsToSkipPackaging" Value="Win64+Linux+Mac" />
 
 <!-- These platforms can't run from loose staged content-->
 <Property Name="PlatformsThatRequirePackaging" Value="Android+IOS" />
