Horde: Move remote execution test script into public folder.

diff --git a/Engine/Build/Graph/Tasks/RemoteExecutionTest.xml b/Engine/Build/Graph/Tasks/RemoteExecutionTest.xml
new file mode 100644
--- /dev/null
+++ b/Engine/Build/Graph/Tasks/RemoteExecutionTest.xml
@@ -0,0 +1,30 @@
+<?xml version='1.0' ?>
+<BuildGraph xmlns="http://www.epicgames.com/BuildGraph" xmlns:xsi="http://www.w3.org/2001/XMLSchema-instance" xsi:schemaLocation="http://www.epicgames.com/BuildGraph ../../../Build/Graph/Schema.xsd" >
+
+	<Option Name="Target" DefaultValue="UnrealPak" Description="Target to build"/>
+	<Option Name="Configuration" DefaultValue="Development" Description="Configuration to build"/>
+	<Option Name="Platforms" DefaultValue="Win64" Description="Platforms to build; Win64+Linux"/>
+	<Option Name="Arguments" DefaultValue="-UBA -UBAForceRemote -UBAStrict" Description="Additional arguments for the build" />
+	<Option Name="Clean" DefaultValue="true" Description="Whether to run UBT with -Clean first" />
+	
+	<Agent Name="Remote Test - Win64" Type="CompileWin64;Win64">
+		<Node Name="Compile Target Win64">
+			<DotNet Arguments="&quot;$(RootDir)/Engine/Binaries/DotNET/UnrealBuildTool/UnrealBuildTool.dll&quot; $(Target) Win64 $(Configuration) -Clean" If="$(Clean)"/>
+			<DotNet Arguments="&quot;$(RootDir)/Engine/Binaries/DotNET/UnrealBuildTool/UnrealBuildTool.dll&quot; $(Target) Win64 $(Configuration) $(Arguments)"/>
+		</Node>
+	</Agent>
+	<Agent Name="Remote Test - Linux" Type="CompileLinux;Linux">
+		<Node Name="Compile Target Linux">
+			<DotNet Arguments="&quot;$(RootDir)/Engine/Binaries/DotNET/UnrealBuildTool/UnrealBuildTool.dll&quot; $(Target) Linux $(Configuration) -Clean" If="$(Clean)"/>
+			<DotNet Arguments="&quot;$(RootDir)/Engine/Binaries/DotNET/UnrealBuildTool/UnrealBuildTool.dll&quot; $(Target) Linux $(Configuration) $(Arguments)"/>
+		</Node>
+	</Agent>
+
+	<Property Name="CompileTargets" Value=""/>
+	<ForEach Name="Platform" Values="$(Platforms)" Separator="+">
+		<Property Name="CompileTargets" Value="$(CompileTargets);Compile Target Win64" If="'$(Platform)' == Win64"/>
+		<Property Name="CompileTargets" Value="$(CompileTargets);Compile Target Linux" If="'$(Platform)' == Linux"/>
+	</ForEach>
+
+	<Aggregate Name="Compile Target" Requires="$(CompileTargets)"/>
+</BuildGraph>
diff --git a/Engine/Source/Programs/Horde/Horde.Server/Defaults/ue5-dev-main-hordedemo.stream.json b/Engine/Source/Programs/Horde/Horde.Server/Defaults/ue5-dev-main-hordedemo.stream.json
--- a/Engine/Source/Programs/Horde/Horde.Server/Defaults/ue5-dev-main-hordedemo.stream.json
+++ b/Engine/Source/Programs/Horde/Horde.Server/Defaults/ue5-dev-main-hordedemo.stream.json
@@ -641,7 +641,7 @@
             "name": "Remote Execution Test",
             "arguments":
             [
-                "-Script=Engine/Restricted/NotForLicensees/Build/RemoteExecutionTest.xml",
+                "-Script=Engine/Build/Graph/Tasks/RemoteExecutionTest.xml",
 				"-Target=Compile Target"
             ],
             "parameters":
