Enforce packaged build through argument instead of path

#rnx

diff --git a/Engine/Build/Graph/Tasks/BuildAndTestProject.xml b/Engine/Build/Graph/Tasks/BuildAndTestProject.xml
--- a/Engine/Build/Graph/Tasks/BuildAndTestProject.xml
+++ b/Engine/Build/Graph/Tasks/BuildAndTestProject.xml
@@ -590,8 +590,9 @@
 				<Property Name="NodeAnnotations" Value=""/>
 
 				<!-- Prerequisites for the test to run -->
+				<Property Name="UsePackaged" Value="(ContainsItem('$(PlatformsThatRequirePackaging)', '$(TargetPlatform)', '+') Or !ContainsItem('$(PlatformsToSkipPackaging)', '$(TargetPlatform)', '+'))" />
 				<Property Name="TestPrerequisites" Value="$(RequirePlatformStaged)" />
-				<Property Name="TestPrerequisites" Value="$(RequirePlatformPackaged)" If="ContainsItem('$(PlatformsThatRequirePackaging)', '$(TargetPlatform)', '+') Or !ContainsItem('$(PlatformsToSkipPackaging)', '$(TargetPlatform)', '+')"/>
+				<Property Name="TestPrerequisites" Value="$(RequirePlatformPackaged)" If="$(UsePackaged)"/>
 				<Property Name="TestPrerequisites" Value="$(PlatformToolsNodeName)" If="'$(GeneratedBuildVersion)' != ''"/>
 				<Property Name="TestPrerequisites" Value="$(PlatformToolsNodeName);$(PlatformEditorNodeName)" If="'$(GeneratedBuildVersion)' != '' and $(TestsRequireEditor)"/>
 				
@@ -601,10 +602,8 @@
 				<Do If="'$(NetworkOutputDirectory)' != '' Or '$(GeneratedBuildVersion)' != ''">
 					<Property Name="TestDirectory" Value="$(NetworkOutputDirectory)" />
 					<Property Name="TestDirectory" Value="$(GeneratedBuildDir)" If="'$(GeneratedBuildVersion)' != ''"/>
-					<Do If="Exists('$(TestDirectory)/$(StagedPlatformFolder)/Packages/$(TestConfiguration)')">
-						<Property Name="TestDirectory" Value="$(TestDirectory)/$(StagedPlatformFolder)/Packages/$(TestConfiguration)" />
-					</Do>
 					<Property Name="GauntletBuildArg" Value="-build=&quot;$(TestDirectory)&quot;" />
+					<Property Name="GauntletBuildArg" Value="$(GauntletBuildArg) -packaged" If="$(UsePackaged)" />
 					<Property Name="GauntletBuildArg" Value="$(GauntletBuildArg) -destlocalinstalldir=&quot;$(LocalInstallDirectory)&quot;" If="!$(IsLocalPlatform) And $(WithDeviceReservation)"/>
 				</Do>
 
diff --git a/Engine/Build/Graph/Tasks/Inc/PlatformOptions.xml b/Engine/Build/Graph/Tasks/Inc/PlatformOptions.xml
--- a/Engine/Build/Graph/Tasks/Inc/PlatformOptions.xml
+++ b/Engine/Build/Graph/Tasks/Inc/PlatformOptions.xml
@@ -156,7 +156,7 @@
 </ForEach>
 
 <!-- Packaging isn't supported by Gauntlet yet :( )-->
-<Property Name="PlatformsToSkipPackaging" Value="Win64+Linux+Mac" />
+<Property Name="PlatformsToSkipPackaging" Value="Win64+Linux" />
 
 <!-- These platforms can't run from loose staged content-->
 <Property Name="PlatformsThatRequirePackaging" Value="Android+IOS" />
