Horde: Rename IncrementalCompileWin64 agent type to IncrementalWin64, matching what we typically use internally.

#rnx

diff --git a/Engine/Build/Graph/Examples/PresubmitTests.xml b/Engine/Build/Graph/Examples/PresubmitTests.xml
--- a/Engine/Build/Graph/Examples/PresubmitTests.xml
+++ b/Engine/Build/Graph/Examples/PresubmitTests.xml
@@ -12,7 +12,7 @@
 	<Option Name="ProjectPath" Restrict="[^ ]*" DefaultValue="" Description="Optional Uproject path to include in the arguments"/>
 	<Option Name="EditorTarget" Restrict="[^ ]+" DefaultValue="UnrealEditor" Description="Name of the editor target to be built"/>
 	
-	<Agent Name="Default Agent" Type="IncrementalCompileWin64;CompileWin64;Win64">
+	<Agent Name="Default Agent" Type="IncrementalWin64;CompileWin64;Win64">
 		<Node Name="Compile $(EditorTarget) Win64" Produces="#EditorBinaries;#TargetInfo">
 			<Compile Target="$(EditorTarget)" Platform="Win64" Configuration="Development" Project="$(ProjectPath)"/>
 		</Node>
diff --git a/Engine/Source/Programs/Horde/Horde.Server/Defaults/ue5-dev-main-hordedemo.stream.json b/Engine/Source/Programs/Horde/Horde.Server/Defaults/ue5-dev-main-hordedemo.stream.json
--- a/Engine/Source/Programs/Horde/Horde.Server/Defaults/ue5-dev-main-hordedemo.stream.json
+++ b/Engine/Source/Programs/Horde/Horde.Server/Defaults/ue5-dev-main-hordedemo.stream.json
@@ -224,7 +224,7 @@
 				"-Target=Submit To Horde For UGS",
 				"-set:UProjectPath=$(ProjectPath)",
 				"-set:EditorTarget=$(EditorTarget)",
-				"-set:PreferredAgent=IncrementalCompileWin64;CompileWin64;Win64"
+				"-set:PreferredAgent=IncrementalWin64;CompileWin64;Win64"
 			]
 		},
 		{
