Horde: Use the new artifacts system for tagging binaries produced for UGS.

#rnx

diff --git a/Engine/Build/Graph/Examples/BuildEditorAndTools.xml b/Engine/Build/Graph/Examples/BuildEditorAndTools.xml
--- a/Engine/Build/Graph/Examples/BuildEditorAndTools.xml
+++ b/Engine/Build/Graph/Examples/BuildEditorAndTools.xml
@@ -136,7 +136,7 @@
 		</Node>	
 		
 		<!-- Prepare a build for UGS -->
-		<Node Name="Stage for UGS" Requires="$(SubmissionDependencies)">
+		<Node Name="Stage for UGS" Requires="$(SubmissionDependencies)" Produces="#PublishBinaries">
 			<!-- Clear out the staging directory -->
 			<Delete Files="$(ArchiveDir)\..."/>
 
@@ -159,6 +159,7 @@
 			<!-- Stage all the files to be archived -->
 			<Copy Files="#ArchiveBinaries" From="$(RootDir)" To="$(ArchiveStagingDir)"/>
 			<Strip Files="#ArchiveSymbols" BaseDir="$(RootDir)" OutputDir="$(ArchiveStagingDir)" Platform="Win64"/>
+			<Tag Files="$(ArchiveStagingDir)/..." With="#PublishBinaries"/>
 		</Node>
 		
 		<!-- Create a zip file and submit it to Perforce -->
@@ -173,9 +174,7 @@
 		</Node>
 
 		<!-- Upload binaries to Horde -->
-		<Node Name="Submit To Horde For UGS" Requires="Stage for UGS">
-			<CreateArtifact Name="editor" Type="ugs-pcb" BaseDir="$(ArchiveStagingDir)"/>
-		</Node>
+		<Artifact Name="editor" Description="Editor PCBs" Type="ugs-pcb" BasePath="$(ArchiveStagingDir)" Tag="#PublishBinaries"/>
 		
 	</Agent>
 
diff --git a/Engine/Source/Programs/Horde/Horde.Server/Defaults/ue5-release-5.4.stream.json b/Engine/Source/Programs/Horde/Horde.Server/Defaults/ue5-release-5.4.stream.json
--- a/Engine/Source/Programs/Horde/Horde.Server/Defaults/ue5-release-5.4.stream.json
+++ b/Engine/Source/Programs/Horde/Horde.Server/Defaults/ue5-release-5.4.stream.json
@@ -221,7 +221,7 @@
 			"arguments":
 			[
 				"-Script=Engine/Build/Graph/Examples/BuildEditorAndTools.xml",
-				"-Target=Submit To Horde For UGS",
+				"-Target=Stage for UGS",
 				"-set:UProjectPath=$(ProjectPath)",
 				"-set:EditorTarget=$(EditorTarget)",
 				"-set:PreferredAgent=IncrementalWin64;CompileWin64;Win64"
