Horde: Fix base path for uploading editor PCBs from BuildEditorAndTools.xml script.

#rnx

diff --git a/Engine/Build/Graph/Examples/BuildEditorAndTools.xml b/Engine/Build/Graph/Examples/BuildEditorAndTools.xml
--- a/Engine/Build/Graph/Examples/BuildEditorAndTools.xml
+++ b/Engine/Build/Graph/Examples/BuildEditorAndTools.xml
@@ -57,6 +57,7 @@
 
 	<Property Name="ArchiveDir" Value="$(RootDir)\LocalBuilds\ArchiveForUGS"/>
 	<Property Name="ArchiveStagingDir" Value="$(ArchiveDir)\Staging"/>
+	<Property Name="ArchiveStagingRelativeDir" Value="LocalBuilds\ArchiveForUGS\Staging"/>
 
 	<EnvVar Name="COMPUTERNAME"/>
 	<EnvVar Name="P4CLIENT"/>
@@ -174,7 +175,7 @@
 		</Node>
 
 		<!-- Upload binaries to Horde -->
-		<Artifact Name="editor" Description="Editor PCBs" Type="ugs-pcb" BasePath="$(ArchiveStagingDir)" Tag="#PublishBinaries"/>
+		<Artifact Name="editor" Description="Editor PCBs" Type="ugs-pcb" BasePath="$(ArchiveStagingRelativeDir)" Tag="#PublishBinaries"/>
 		
 	</Agent>
 
