Allows pre-cook steps to be injected

diff --git a/Engine/Build/Graph/Tasks/BuildAndTestProject.xml b/Engine/Build/Graph/Tasks/BuildAndTestProject.xml
--- a/Engine/Build/Graph/Tasks/BuildAndTestProject.xml
+++ b/Engine/Build/Graph/Tasks/BuildAndTestProject.xml
@@ -114,9 +114,13 @@
 
 	<Option Name="TestsRequireEditor" DefaultValue="false" Description="Whether to force tests to require editor compilation"/>
 
+	<!-- Macros that allows custom pre-steps -->
+	<Macro Name="PreCookStep"/>
+
 	<!-- Settings files that can be overriden -->
 	<Option Name="PathSettingsFile" DefaultValue="Inc/PathSettings.xml" Description="Path settings to use" />
 	<Option Name="GauntletSettingsFile" DefaultValue="Inc/GauntletSettings.xml" Description="Gauntlet settings to use" />
+	<Option Name="PreStepsExtendFile" DefaultValue="" Description="Path to the file containing all pre-step macro extends"/>
 
     <!-- Annotations notif settings -->
 	<Option Name="AnnotationsTarget" DefaultValue="" Description="Set annotations target" />
@@ -129,6 +133,10 @@
 	<!-- Common settings but can be overidden -->
 	<Include Script="$(PathSettingsFile)" />
 	<Include Script="$(GauntletSettingsFile)" />
+	<Do If="'$(PreStepsExtendFile)' != ''">
+		<Error Message="Pre Steps Extend file $(PreStepsExtendFile) does not exist" If="!Exists('$(PreStepsExtendFile)')"/>
+		<Include Script="$(PreStepsExtendFile)"/>
+	</Do>
 
 	<!-- At this point we have all options so calculate a few things -->
 
@@ -524,6 +532,7 @@
 			<!-- Cook -->
 			<Node Name="$(PlatformCookNodeName)" Requires="$(PlatformEditorNodeName)" Produces="#$(PlatformCookNodeName) Complete" RunEarly="true">
 				<Do If="!$(SkipCook)">
+					<Expand Name="PreCookStep"/>
 					<Property Name="CookArgs" Value="" />
 					<Property Name="CookArgs" Value="-MapIniSection=$(MapIniSection) $(CookPlatform)" If="'$(MapIniSection)' != ''"/>				
 					<Cook Project="$(TargetProject)" Platform="$(CookPlatform)" TagOutput="false" Arguments="$(CookArgs) $(ExtraCookArguments)"/>
