Add Android to the EngineTest schedule

#rnx

diff --git a/Engine/Build/Graph/Tasks/Inc/SharedProperties.xml b/Engine/Build/Graph/Tasks/Inc/SharedProperties.xml
--- a/Engine/Build/Graph/Tasks/Inc/SharedProperties.xml
+++ b/Engine/Build/Graph/Tasks/Inc/SharedProperties.xml
@@ -12,6 +12,8 @@
 	<!-- URL to Horde server -->
 	<Property Name="HordeDashboardURL" Value=""/>
 	<Option Name="GauntletDeviceService" DefaultValue="" Description="Gauntlet Device Service to use for device reservation when running tests"/>
+	<!-- Android adbkeys -->
+	<Option Name="AdbKeysPath" DefaultValue="" Description="Path to adbkeys location (for Android)"/>
 	<!-- Get custom shared properties -->
 	<Do If="Exists('$(RootDir)/$(SharedPropertiesFile)')">
 		<Include Script="$(RootDir)/$(SharedPropertiesFile)"/>
