Horde: Add the ugs-project key to artifacts produced by BuildEditorAndTools.xml.

diff --git a/Engine/Build/Graph/Examples/BuildEditorAndTools.xml b/Engine/Build/Graph/Examples/BuildEditorAndTools.xml
--- a/Engine/Build/Graph/Examples/BuildEditorAndTools.xml
+++ b/Engine/Build/Graph/Examples/BuildEditorAndTools.xml
@@ -175,7 +175,10 @@
 		</Node>
 
 		<!-- Upload binaries to Horde -->
-		<Artifact Name="editor" Description="Editor PCBs" Type="ugs-pcb" BasePath="$(ArchiveStagingRelativeDir)" Keys="target:$(EditorTarget)" Tag="#PublishBinaries"/>
+		<Property Name="UgsProjectKey" Value=""/>
+		<Property Name="UgsProjectKey" Value="ugs-project=$(Branch)/$(UProjectPath)" If="'$(Branch)' != '' and '$(UProjectPath)' != ''"/>
+		<Property Name="UgsProjectKey" Value="ugs-project=$(Branch)" If="'$(Branch)' != '' and '$(UgsProjectKey)' == ''"/>
+		<Artifact Name="editor" Description="Editor PCBs" Type="ugs-pcb" BasePath="$(ArchiveStagingRelativeDir)" Keys="$(UgsProjectKey)" Metadata="ArchiveType=Editor" Tag="#PublishBinaries"/>
 		
 	</Agent>
 
