Horde: Fix BuildAndTestProject.xml making assumptions about $(IsBuildMachine) meaning that we're running in an Epic-internal environment. Also publish artifacts to Horde.

diff --git a/Engine/Build/Graph/Tasks/BuildAndTestProject.xml b/Engine/Build/Graph/Tasks/BuildAndTestProject.xml
--- a/Engine/Build/Graph/Tasks/BuildAndTestProject.xml
+++ b/Engine/Build/Graph/Tasks/BuildAndTestProject.xml
@@ -460,6 +460,7 @@
 		<Property Name="PlatformCookNodeName" Value="$(PreNodeName)Cook $(TargetPlatform)" />
 		<Property Name="PlatformStageNodeName" Value="$(PreNodeName)Stage $(TargetPlatform)" />
 		<Property Name="PlatformPackageNodeName" Value="$(PreNodeName)Package $(TargetPlatform)" />
+		<Property Name="PlatformPackageArtifactName" Value="$(PlatformPackageNodeName)" />
 		<Property Name="PlatformHashNodeName" Value="$(PreNodeName)Hash $(TargetPlatform)" />
 		<Property Name="PlatformPublishStagedNodeName" Value="$(PreNodeName)Publish Staged $(TargetPlatform)" />
 		<Property Name="PlatformPublishPackagedNodeName" Value="$(PreNodeName)Publish Packaged $(TargetPlatform)" />
@@ -481,7 +482,7 @@
 		<!-- <Property Name="StagedPlatformFolder" Value="$(StagedPlatformFolder)Client" If="'$(TargetType)' == 'Client'" /> -->
 
 		<!-- Hashing of staged data can only be done on Win64 hosts -->
-		<Property Name="DoUnsyncHashing" Value="$(HostAgentType) == 'Win64'" />
+		<Property Name="DoUnsyncHashing" Value="$(HostAgentType) == 'Win64' and $(IsEpicInternal)" />
 
 		<!-- Default BuildCookRun args needed when staging/packaging/archiving this configuration -->
 		<Property Name="BCRArgs" Value="-project=$(TargetProject) -platform=$(TargetPlatform) -cookflavor=&quot;$(BCRCookFlavor)&quot;" />
@@ -540,7 +541,7 @@
 			</Node>
 			
 			<!-- Stage -->
-			<Node Name="$(PlatformStageNodeName)" Requires="$(PlatformCompileNodeName);#$(PlatformCookNodeName) Complete">
+			<Node Name="$(PlatformStageNodeName)" Requires="$(PlatformCompileNodeName);#$(PlatformCookNodeName) Complete" Produces="#$(PlatformStageNodeName) Files">
 				<Switch>
 					<Case If="$(IsMobilePlatform)">
 						<ForEach Name="PackageConfiguration" Values="$(PackageConfigurations)" Separator="+">
@@ -551,8 +552,12 @@
 						<Command Name="BuildCookRun" Arguments="$(BCRArgs) -configuration=$(PackageConfigurations) $(StageAndPackageArguments) -skipbuild -skipcook -stage -pak" />
 					</Default>
 				</Switch>
+				<Tag Files="$(ProjectOutputDirectory)/$(StagedPlatformFolder)/..." With="#$(PlatformStageNodeName) Files"/>
 			</Node>
 			
+			<!-- Upload to Horde -->
+			<Artifact Name="$(PlatformPackageArtifactName)" Type="packaged-build" BasePath="$(ProjectOutputDirectory)/$(StagedPlatformFolder)" Tag="#$(PlatformStageNodeName) Files"/>
+			
 			<!-- Hash -->
 			<Node Name="$(PlatformHashNodeName)" Requires="$(PlatformStageNodeName)">
 				<Do If="$(DoUnsyncHashing)">
@@ -585,6 +590,7 @@
 			
 			<!-- Publish (Packages) -->
 			<Node Name="$(PlatformPublishPackagedNodeName)" Requires="$(PlatformPackageNodeName)">
+				<Error Message="NetworkOutputDirectory is not set" If="'$(NetworkOutputDirectory)' == ''"/>
 				<ForEach Name="PackageConfiguration" Values="$(PackageConfigurations)" Separator="+">
 					<Command Name="BuildCookRun" Arguments="$(BCRArgs) -configuration=$(PackageConfiguration) $(StageAndPackageArguments) -skipbuild -skipcook -skipstage -skippak -skippackage -archive -archivedirectory=&quot;$(NetworkOutputDirectory)/$(StagedPlatformFolder)/Packages&quot;" />
 				</ForEach>
diff --git a/Engine/Build/Graph/Tasks/Inc/PathSettings.xml b/Engine/Build/Graph/Tasks/Inc/PathSettings.xml
--- a/Engine/Build/Graph/Tasks/Inc/PathSettings.xml
+++ b/Engine/Build/Graph/Tasks/Inc/PathSettings.xml
@@ -31,14 +31,20 @@
 
 	-->
 
+	<Property Name="IsEpicInternal" Value="false"/>
+	<Property Name="IsEpicInternal" Value="true" If="Exists('Engine/Restricted/NotForLicensees/Build/EpicInternal.txt')"/>
+
+	<Property Name="IsEpicBuildMachine" Value="false"/>
+	<Property Name="IsEpicBuildMachine" Value="true" If="$(IsBuildMachine) and $(IsEpicInternal)"/>
+
 	<!-- Publish packaged builds to the network ?-->
-	<Option Name="PublishBuilds" DefaultValue="$(IsBuildMachine)" Description="Publish packaged builds to the builds folder on the network"/>
+	<Option Name="PublishBuilds" DefaultValue="$(IsEpicBuildMachine)" Description="Publish packaged builds to the builds folder on the network"/>
 
 	<!-- Publish a hashed staging data that developers can quickly sync to -->
 	<Option Name="PublishHashedStagingData" DefaultValue="false" Restrict="true|false" Description="Create hashes for syncing via Unsync" />
 
 	<!-- Publish automation reports to the network-->
-	<Option Name="PublishReports" DefaultValue="$(IsBuildMachine)" Description="Publish packaged builds to the builds folder on the netwok"/>
+	<Option Name="PublishReports" DefaultValue="$(IsEpicBuildMachine)" Description="Publish packaged builds to the builds folder on the netwok"/>
 
 	<!-- For pre-generated build -->
 	<Option Name="GeneratedBuildVersion" DefaultValue="" Description="The pre-generated build version"/>
@@ -70,7 +76,7 @@
 	<Property Name="NetworkOutputDirectory" Value="$(NetworkOutputDirectory)/$(EscapedBranch)/$(BuildNamePath)" If="'$(NetworkOutputDirectory)' != ''"/>
 
 	<!-- Automation can publish reports to a folder that is mapped to a website -->
-	<Do If="$(IsBuildMachine) or $(PublishReports)" >
+	<Do If="$(IsEpicBuildMachine) or $(PublishReports)" >
 		<Property Name="AutomationReportOutputDirectory" Value="$(NetworkReportRoot)"/>
 		<Property Name="AutomationReportOutputDirectory" Value="$(AutomationReportOutputDirectory)/$(EscapedBranch)/$(ProjectName)/$(CL-String)"/>
 		<Property Name="AutomationReportUrl" Value="http://automation.epicgames.net/reports/$(EscapedBranch)/$(ProjectName)/$(CL-String)"/>
diff --git a/Engine/Source/Programs/AutomationTool/BuildGraph/BgScriptReader.cs b/Engine/Source/Programs/AutomationTool/BuildGraph/BgScriptReader.cs
--- a/Engine/Source/Programs/AutomationTool/BuildGraph/BgScriptReader.cs
+++ b/Engine/Source/Programs/AutomationTool/BuildGraph/BgScriptReader.cs
@@ -1244,6 +1244,10 @@ namespace AutomationTool
 				{
 					tag = $"#{name}";
 				}
+				if (!_graph.TagNameToNodeOutput.TryGetValue(tag, out _))
+				{
+					LogError(element, "Artifact '{Name}' references non-existent tag '{Tag}'", name, tag);
+				}
 
 				string[] keys = ReadListAttribute(element, "Keys");
 				string[] metadata = ReadListAttribute(element, "Metadata");
diff --git a/Engine/Source/Programs/Shared/EpicGames.BuildGraph/BgGraphDef.cs b/Engine/Source/Programs/Shared/EpicGames.BuildGraph/BgGraphDef.cs
--- a/Engine/Source/Programs/Shared/EpicGames.BuildGraph/BgGraphDef.cs
+++ b/Engine/Source/Programs/Shared/EpicGames.BuildGraph/BgGraphDef.cs
@@ -288,7 +288,7 @@ namespace EpicGames.BuildGraph
 			}
 			NameToAggregate = newNameToAggregate;
 
-			// Remove any labels that are no longer value
+			// Remove any labels that are no longer valid
 			foreach (BgLabelDef label in Labels)
 			{
 				label.RequiredNodes.RemoveWhere(x => !retainNodes.Contains(x));
@@ -298,6 +298,17 @@ namespace EpicGames.BuildGraph
 
 			// Remove any badges which do not have all their dependencies
 			Badges.RemoveAll(x => x.Nodes.Any(y => !retainNodes.Contains(y)));
+
+			// Rebuild the tag name to output dictionary
+			Dictionary<string, BgNodeOutput> newTagNameToNodeOutput = new Dictionary<string, BgNodeOutput>(TagNameToNodeOutput.Comparer);
+			foreach (BgNodeOutput output in Agents.SelectMany(x => x.Nodes).SelectMany(x => x.Outputs))
+			{
+				newTagNameToNodeOutput.Add(output.TagName, output);
+			}
+			TagNameToNodeOutput = newTagNameToNodeOutput;
+
+			// Remove any artifacts which are not produced
+			Artifacts.RemoveAll(x => !TagNameToNodeOutput.ContainsKey(x.TagName));
 		}
 
 		/// <summary>
@@ -721,6 +732,18 @@ namespace EpicGames.BuildGraph
 				}
 				logger.LogInformation("");
 			}
+
+			// Print all the produced artifacts
+			BgArtifactDef[] artifacts = Artifacts.OrderBy(x => x.Name).ToArray();
+			if (artifacts.Length > 0)
+			{
+				logger.LogInformation("Artifacts:");
+				foreach (BgArtifactDef artifact in artifacts.OrderBy(x => x.Description))
+				{
+					logger.LogInformation("    {Name}", artifact.Name);
+				}
+				logger.LogInformation("");
+			}
 		}
 	}
 
