Add capability to publish snapshots to both zen and jupiter.

diff --git a/Engine/Build/Graph/Tasks/BuildAndTestProject.xml b/Engine/Build/Graph/Tasks/BuildAndTestProject.xml
--- a/Engine/Build/Graph/Tasks/BuildAndTestProject.xml
+++ b/Engine/Build/Graph/Tasks/BuildAndTestProject.xml
@@ -7,6 +7,7 @@
 	<EnvVar Name="UE-CloudPublishDescriptorHost"/>
 	<EnvVar Name="UE-CloudPublishDescriptorHttpVersion"/>
 	<EnvVar Name="UE-CloudPublishNamespace"/>
+	<EnvVar Name="UE-ZenPublishHost"/>
 
 	<Option Name="Horde" Restrict="true|false" DefaultValue="$(IsBuildMachine)" Description="Whether the build is running under Horde"/>
 	
@@ -567,12 +568,23 @@
 
 			<!-- (Optional) Publish to Cloud -->
 			<Property Name="PlatformPublishStageDependencyName" Value="$(PlatformHashNodeName)" />
-			<Do If="$(PublishBuilds) and $(PublishCloudSnapshot) and '$(UE-CloudPublishHost)' != '' and '$(UE-CloudPublishNamespace)' != ''">
-				<Property Name="PlatformCloudPublishNodeName" Value="$(PreNodeName)Cloud Publish $(TargetPlatform)" />
-				<Node Name="$(PlatformCloudPublishNodeName)" Requires="$(PlatformHashNodeName)">
-					<ZenExportSnapshot Project="$(TargetProject)" Platform="$(CookPlatform)" DestinationStorageType="Cloud" DestinationCloudHost="$(UE-CloudPublishHost)" DestinationCloudHttpVersion="$(UE-CloudPublishHttpVersion)" SnapshotDescriptorCloudHost="$(UE-CloudPublishDescriptorHost)" SnapshotDescriptorCloudHttpVersion="$(UE-CloudPublishDescriptorHttpVersion)" DestinationCloudNamespace="$(UE-CloudPublishNamespace)" DestinationIdentifier="$(EscapedBranch)/$(BuildNamePath)" SnapshotDescriptorFile="$(NetworkOutputDirectory)/$(StagedPlatformFolder)/SnapshotDescriptor.json"/>
-				</Node>
-				<Property Name="PlatformPublishStageDependencyName" Value="$(PlatformCloudPublishNodeName)" />
+			<Do If="$(PublishBuilds) and $(PublishCloudSnapshot)">
+				<Switch>
+					<Case If="'$(UE-ZenPublishHost)' != ''">
+						<Property Name="PlatformSnapshotPublishNodeName" Value="$(PreNodeName)Snapshot Publish $(TargetPlatform)" />
+						<Node Name="$(PlatformSnapshotPublishNodeName)" Requires="$(PlatformHashNodeName)">
+							<ZenExportSnapshot Project="$(TargetProject)" Platform="$(CookPlatform)" DestinationStorageType="Zen" DestinationZenHost="$(UE-ZenPublishHost)" DestinationIdentifier="$(EscapedBranch)/$(BuildNamePath)" SnapshotDescriptorFile="$(NetworkOutputDirectory)/$(StagedPlatformFolder)/SnapshotDescriptor.json"/>
+						</Node>
+						<Property Name="PlatformPublishStageDependencyName" Value="$(PlatformSnapshotPublishNodeName)" />
+					</Case>
+					<Case If="'$(UE-CloudPublishHost)' != '' and '$(UE-CloudPublishNamespace)' != ''">
+						<Property Name="PlatformSnapshotPublishNodeName" Value="$(PreNodeName)Snapshot Publish $(TargetPlatform)" />
+						<Node Name="$(PlatformSnapshotPublishNodeName)" Requires="$(PlatformHashNodeName)">
+							<ZenExportSnapshot Project="$(TargetProject)" Platform="$(CookPlatform)" DestinationStorageType="Cloud" DestinationCloudHost="$(UE-CloudPublishHost)" DestinationCloudHttpVersion="$(UE-CloudPublishHttpVersion)" SnapshotDescriptorCloudHost="$(UE-CloudPublishDescriptorHost)" SnapshotDescriptorCloudHttpVersion="$(UE-CloudPublishDescriptorHttpVersion)" DestinationCloudNamespace="$(UE-CloudPublishNamespace)" DestinationIdentifier="$(EscapedBranch)/$(BuildNamePath)" SnapshotDescriptorFile="$(NetworkOutputDirectory)/$(StagedPlatformFolder)/SnapshotDescriptor.json"/>
+						</Node>
+						<Property Name="PlatformPublishStageDependencyName" Value="$(PlatformSnapshotPublishNodeName)" />
+					</Case>
+				</Switch>
 			</Do>
 			
 			<!-- Publish (Staged) -->
