Fix iterative validate single phase mode so that it depends on the compilation of the editor and tools.

diff --git a/Engine/Build/Graph/Tests/IterativeValidate.xml b/Engine/Build/Graph/Tests/IterativeValidate.xml
--- a/Engine/Build/Graph/Tests/IterativeValidate.xml
+++ b/Engine/Build/Graph/Tests/IterativeValidate.xml
@@ -39,13 +39,13 @@
 			<Compile Target="ShaderCompileWorker" Platform="$(EditorPlatform)" Project="$(ProjectFile)" Configuration="Development"/>
 			<Compile Target="$(EditorTarget)" Platform="$(EditorPlatform)" Project="$(ProjectFile)" Configuration="Development"/>
 		</Node>
-		<Node Name="$(CookIterativeValidateSinglePhaseName)">
+		<Node Name="$(CookIterativeValidateSinglePhaseName)" Requires="$(CompileNodeName)">
 			<Commandlet Name="Cook" Project="$(ProjectFile)" Arguments="$(CommonCookArgs) -IterativeValidate -IterativeValidateAllowWrite"/>
 		</Node>
-		<Node Name="$(CookIterativeValidatePhase1Name)" Requires="$(CompileNodeName)"  If="'$(ValidateMode)' == 'twophase'">
+		<Node Name="$(CookIterativeValidatePhase1Name)" Requires="$(CompileNodeName)">
 			<Commandlet Name="Cook" Project="$(ProjectFile)" Arguments="$(CommonCookArgs) -IterativeValidatePhase1 -IterativeValidateAllowWrite"/>
 		</Node>
-		<Node Name="$(CookIterativeValidatePhase2Name)" Requires="$(CookIterativeValidatePhase1Name)"  If="'$(ValidateMode)' == 'twophase'">
+		<Node Name="$(CookIterativeValidatePhase2Name)" Requires="$(CookIterativeValidatePhase1Name)">
 			<Commandlet Name="Cook" Project="$(ProjectFile)" Arguments="$(CommonCookArgs) -IterativeValidatePhase2 -IterativeValidateAllowWrite"/>
 		</Node>
 	</Agent>
