Allow use of UBA when compiling for iterative validtion jobs.

diff --git a/Engine/Build/Graph/Tests/IterativeValidate.xml b/Engine/Build/Graph/Tests/IterativeValidate.xml
--- a/Engine/Build/Graph/Tests/IterativeValidate.xml
+++ b/Engine/Build/Graph/Tests/IterativeValidate.xml
@@ -17,6 +17,10 @@
 	<Option Name="AgentType" DefaultValue="" Description="The platform to cook"/>
 	<Property Name="AgentType" Value="IterativeValidateAgent_$(ProjectName)_$(TargetPlatform)" If="'$(AgentType)' == ''"/>
 
+	<Option Name="WithUBA" Restrict="true|false" DefaultValue="true" Description="Whether to use UBA compile or not"/>
+	<Property Name="UBACompileArgs" Value=""/>
+	<Property Name="UBACompileArgs" Value="-UBA -ActionStallReportTime=300 -UBALinkRemote -UBAPrintSummary" If="$(WithUBA)"/>
+	
 	<!-- Iterative Options -->
 	<Option Name="IterativeMode" DefaultValue="incremental" Restrict="iterative|incremental" Description="Whether to iteratively cook using Incremental or using the legacy iterative"/>
 	<Option Name="ValidateMode" DefaultValue="onephase" Restrict="onephase|twophase" Description="Whether to validate results in onephase (which reports indeterminism as an iterativefalseskip error) or two phase (which detects indeterminism and downgrades it to a display message)"/>
@@ -36,8 +40,8 @@
 
 	<Agent Name="IterativeValidateAgent" Type="$(AgentType)">
 		<Node Name="$(CompileNodeName)">
-			<Compile Target="ShaderCompileWorker" Platform="$(EditorPlatform)" Project="$(ProjectFile)" Configuration="Development"/>
-			<Compile Target="$(EditorTarget)" Platform="$(EditorPlatform)" Project="$(ProjectFile)" Configuration="Development"/>
+			<Compile Target="ShaderCompileWorker" Platform="$(EditorPlatform)" Project="$(ProjectFile)" Configuration="Development" Arguments="$(UBACompileArgs)"/>
+			<Compile Target="$(EditorTarget)" Platform="$(EditorPlatform)" Project="$(ProjectFile)" Configuration="Development" Arguments="$(UBACompileArgs)"/>
 		</Node>
 		<Node Name="$(CookIterativeValidateSinglePhaseName)" Requires="$(CompileNodeName)">
 			<Commandlet Name="Cook" Project="$(ProjectFile)" Arguments="$(CommonCookArgs) -IterativeValidate -IterativeValidateAllowWrite"/>
