Horde: Set the compatible changelist when building editor binaries.

diff --git a/Engine/Build/Graph/Examples/BuildEditorAndTools.xml b/Engine/Build/Graph/Examples/BuildEditorAndTools.xml
--- a/Engine/Build/Graph/Examples/BuildEditorAndTools.xml
+++ b/Engine/Build/Graph/Examples/BuildEditorAndTools.xml
@@ -66,7 +66,7 @@
 	
 		<!-- Update the engine version files -->
 		<Node Name="Update Version Files">
-			<SetVersion Change="$(Change)" Branch="$(EscapedBranch)" Licensee="$(Licensee)" If="$(Versioned)"/>
+			<SetVersion Change="$(Change)" CompatibleChange="$(CodeChange)" Branch="$(EscapedBranch)" Licensee="$(Licensee)" If="$(Versioned)"/>
 			<SetVersion Change="0" Branch="$(EscapedBranch)" Licensee="$(Licensee)" If="!$(Versioned) and $(Promoted)"/>
 		</Node>
 		
