Tag files for artifacts in a separate node to the one that produces them, to work around package step cloberring the outputs.

#rnx

diff --git a/Engine/Build/Graph/Tasks/BuildAndTestProject.xml b/Engine/Build/Graph/Tasks/BuildAndTestProject.xml
--- a/Engine/Build/Graph/Tasks/BuildAndTestProject.xml
+++ b/Engine/Build/Graph/Tasks/BuildAndTestProject.xml
@@ -465,8 +465,9 @@
 		<Property Name="PlatformHashNodeName" Value="$(PreNodeName)Hash $(TargetPlatform)" />
 		<Property Name="PlatformPublishStagedNodeName" Value="$(PreNodeName)Publish Staged $(TargetPlatform)" />
 		<Property Name="PlatformPublishPackagedNodeName" Value="$(PreNodeName)Publish Packaged $(TargetPlatform)" />
-		<Property Name="RequirePlatformStaged" Value="$(PlatformStageNodeName)" />
-		<Property Name="RequirePlatformStaged" Value="$(PlatformPublishStagedNodeName)" If="!$(SkipPublish) And $(PublishBuilds)"/>
+		<Property Name="PlatformStagedArtifactNodeName" Value="$(PlatformPublishStagedNodeName) (Horde)" />
+		<Property Name="RequirePlatformStaged" Value="$(PlatformStageNodeName);$(PlatformStagedArtifactNodeName)" />
+		<Property Name="RequirePlatformStaged" Value="$(RequirePlatformStaged);$(PlatformPublishStagedNodeName)" If="!$(SkipPublish) And $(PublishBuilds)"/>
 		<Property Name="RequirePlatformPackaged" Value="$(PlatformPackageNodeName)" />
 		<Property Name="RequirePlatformPackaged" Value="$(PlatformPublishPackagedNodeName)"  If="!$(SkipPublish) And $(PublishBuilds)"/>
 		
@@ -542,7 +543,7 @@
 			</Node>
 			
 			<!-- Stage -->
-			<Node Name="$(PlatformStageNodeName)" Requires="$(PlatformCompileNodeName);#$(PlatformCookNodeName) Complete" Produces="#$(PlatformStageNodeName) Files">
+			<Node Name="$(PlatformStageNodeName)" Requires="$(PlatformCompileNodeName);#$(PlatformCookNodeName) Complete">
 				<Switch>
 					<Case If="$(IsMobilePlatform)">
 						<ForEach Name="PackageConfiguration" Values="$(PackageConfigurations)" Separator="+">
@@ -553,12 +554,8 @@
 						<Command Name="BuildCookRun" Arguments="$(BCRArgs) -configuration=$(PackageConfigurations) $(StageAndPackageArguments) -skipbuild -skipcook -stage -pak" />
 					</Default>
 				</Switch>
-				<Tag Files="$(ProjectOutputDirectory)/$(StagedPlatformFolder)/..." With="#$(PlatformStageNodeName) Files"/>
 			</Node>
 			
-			<!-- Upload to Horde -->
-			<Artifact Name="$(PlatformPackageArtifactName)" Type="packaged-build" BasePath="$(ProjectOutputDirectory)/$(StagedPlatformFolder)" Tag="#$(PlatformStageNodeName) Files"/>
-			
 			<!-- Hash -->
 			<Node Name="$(PlatformHashNodeName)" Requires="$(PlatformStageNodeName)">
 				<Do If="$(DoUnsyncHashing)">
@@ -592,7 +589,13 @@
 				<Copy Files="..." From="$(ProjectOutputDirectory)/$(StagedPlatformFolder)" To="$(NetworkOutputDirectory)/$(StagedPlatformFolder)/Staged"/>
 				<Copy Files="..." From="$(ShaderSymbolsDirectory)/*$(TargetPlatform)*" To="$(NetworkOutputDirectory)/$(StagedPlatformFolder)/ShaderSymbols"/>
 			</Node>
-			
+
+			<!-- Publish (Staged) (Horde) -->
+			<Node Name="$(PlatformStagedArtifactNodeName)" Requires="$(PlatformStageNodeName)" Produces="#$(PlatformStageNodeName) Files">
+				<Tag Files="$(ProjectOutputDirectory)/$(StagedPlatformFolder)/..." With="#$(PlatformStageNodeName) Files"/>
+			</Node>
+			<Artifact Name="$(PlatformPackageArtifactName)" Type="packaged-build" BasePath="$(ProjectOutputDirectory)/$(StagedPlatformFolder)" Tag="#$(PlatformStageNodeName) Files"/>
+
 			<!-- Package -->
 			<Node Name="$(PlatformPackageNodeName)" Requires="$(PlatformPublishStageDependencyName)">
 				<ForEach Name="PackageConfiguration" Values="$(PackageConfigurations)" Separator="+">
