Allow IterativeValidate jobs to run multi-process.

diff --git a/Engine/Build/Graph/Tests/IterativeValidate.xml b/Engine/Build/Graph/Tests/IterativeValidate.xml
--- a/Engine/Build/Graph/Tests/IterativeValidate.xml
+++ b/Engine/Build/Graph/Tests/IterativeValidate.xml
@@ -30,8 +30,6 @@
 	<!-- Agent=IterativeValidate variables -->
 	<Property Name="EditorPlatform" Value="Win64"/>
 	<Property Name="CommonCookArgs" Value="-targetplatform=$(TargetPlatform) $(IterativeArgs) -skiponlyeditoronly=true"/>
-	<!-- Force singleprocess cook because IterativeValidatePackageWriter includes data that is not replicated -->
-	<Property Name="CommonCookArgs" Value="$(CommonCookArgs) -cookprocesscount=1"/>
 
 	<Property Name="CompileNodeName" Value="Compile $(EditorPlatform) $(ProjectName)"/>
 	<Property Name="CookIterativeValidateSinglePhaseName" Value="Cook IterativeValidateSinglePhase $(ProjectName) $(TargetPlatform)"/>
