Add scheduled job to do verification of determinism for DDC values.

diff --git a/Engine/Build/Graph/Tests/DDCVerify.xml b/Engine/Build/Graph/Tests/DDCVerify.xml
new file mode 100644
--- /dev/null
+++ b/Engine/Build/Graph/Tests/DDCVerify.xml
@@ -0,0 +1,49 @@
+<?xml version='1.0' ?>
+<BuildGraph xmlns="http://www.epicgames.com/BuildGraph" xmlns:xsi="http://www.w3.org/2001/XMLSchema-instance" xsi:schemaLocation="http://www.epicgames.com/BuildGraph ../Schema.xsd" >
+
+	<!-- Project Options -->
+	<Option Name="ProjectFile" DefaultValue="_Unknown_" Restrict=".+" Description="The projects to cook"/>
+	<Option Name="ProjectName" DefaultValue="_Unknown_" Restrict=".+" Description="The projects to cook"/>
+	<Option Name="EditorTarget" DefaultValue="$(ProjectName)Editor" Description="Compile name of the editor target"/>
+
+	<!-- Note we can not make the default='', because the default is used for some validation. -->
+	<Error Message="No project file specified" If="'$(ProjectFile)' == '_Unknown_'"/>
+	<Error Message="No project name specified" If="'$(ProjectName)' == '_Unknown_'"/>
+
+	<!-- Platform Options -->
+	<Option Name="TargetPlatform" DefaultValue="Windows" Restrict=".+" Description="The platform to cook"/>
+
+	<!-- Agent Options -->
+	<Option Name="AgentType" DefaultValue="" Description="The agent type to use"/>
+	<Property Name="AgentType" Value="DDCVerifyAgent_$(ProjectName)_$(TargetPlatform)" If="'$(AgentType)' == ''"/>
+
+	<Option Name="WithUBA" Restrict="true|false" DefaultValue="true" Description="Whether to use UBA compile or not"/>
+	<Property Name="UBACompileArgs" Value=""/>
+	<Property Name="UBACompileArgs" Value="-UBA -ActionStallReportTime=300 -UBALinkRemote -UBAPrintSummary" If="$(WithUBA)"/>
+	
+	<!-- DDC Verify Options -->
+	<Option Name="DDCVerifyDefaultRate" Restrict="\d*" DefaultValue="10" Description="The default verification rate in percent"/>
+	<Option Name="DDCVerifyCustomBucketRates" DefaultValue="Texture=2+StaticMesh=2+MaterialShaderMap=5+fshaderjobcacheshaders=5" Description="Bucket names and associated custom verification rates that differ from the default rate.  Uses plus or comma as separators."/>
+	<Property Name="DDCVerifyArgs" Value="-DDC-Verify"/>
+	<Property Name="DDCVerifyArgs" Value="$(DDCVerifyArgs)=$(DDCVerifyCustomBucketRates)" If="'$(DDCVerifyCustomBucketRates)' != ''"/>
+	<Property Name="DDCVerifyArgs" Value="$(DDCVerifyArgs) -DDC-VerifyRate=$(DDCVerifyDefaultRate)"/>
+
+	<!--========================================-->
+	<Property Name="EditorPlatform" Value="Win64"/>
+	<Property Name="CommonCookArgs" Value="-targetplatform=$(TargetPlatform) $(DDCVerifyArgs) -skiponlyeditoronly=true"/>
+
+	<Property Name="CompileNodeName" Value="Compile $(EditorPlatform) $(ProjectName)"/>
+
+	<Agent Name="DDCVerifyAgent" Type="$(AgentType)">
+		<Node Name="$(CompileNodeName)">
+			<Compile Target="ShaderCompileWorker" Platform="$(EditorPlatform)" Project="$(ProjectFile)" Configuration="Development" Arguments="$(UBACompileArgs)"/>
+			<Compile Target="$(EditorTarget)" Platform="$(EditorPlatform)" Project="$(ProjectFile)" Configuration="Development" Arguments="$(UBACompileArgs)"/>
+		</Node>
+		<Node Name="DDC Verify" Requires="$(CompileNodeName)">
+			<Commandlet Name="Cook" Project="$(ProjectFile)" Arguments="$(CommonCookArgs)"/>
+		</Node>
+	</Agent>
+
+	<Aggregate Name="DDCVerify" Requires="DDC Verify"/>
+
+</BuildGraph>
\ No newline at end of file
