Reduce the default verification rates for the DDC-Verify job to allow its duration to be suitable for large project use.

diff --git a/Engine/Build/Graph/Tests/DDCVerify.xml b/Engine/Build/Graph/Tests/DDCVerify.xml
--- a/Engine/Build/Graph/Tests/DDCVerify.xml
+++ b/Engine/Build/Graph/Tests/DDCVerify.xml
@@ -22,8 +22,8 @@
 	<Property Name="UBACompileArgs" Value="-UBA -ActionStallReportTime=300 -UBALinkRemote -UBAPrintSummary" If="$(WithUBA)"/>
 	
 	<!-- DDC Verify Options -->
-	<Option Name="DDCVerifyDefaultRate" Restrict="\d*" DefaultValue="10" Description="The default verification rate in percent"/>
-	<Option Name="DDCVerifyCustomBucketRates" DefaultValue="Texture=2+StaticMesh=2+MaterialShaderMap=5+fshaderjobcacheshaders=5" Description="Bucket names and associated custom verification rates that differ from the default rate.  Uses plus or comma as separators."/>
+	<Option Name="DDCVerifyDefaultRate" Restrict="\d*" DefaultValue="5" Description="The default verification rate in percent"/>
+	<Option Name="DDCVerifyCustomBucketRates" DefaultValue="Texture=2+SkeletalMesh=2+StaticMesh=2+NaniteDisplacedMesh=2+Card=2+DistanceField=2+SubUV=2+GlobalShaderMap=2+MaterialShaderMap=2+fshaderjobcacheshaders=2+Audio=2+StreamedAudio=2+AnimationSequence=2" Description="Bucket names and associated custom verification rates that differ from the default rate.  Uses plus or comma as separators."/>
 	<Property Name="DDCVerifyArgs" Value="-DDC-Verify"/>
 	<Property Name="DDCVerifyArgs" Value="$(DDCVerifyArgs)=$(DDCVerifyCustomBucketRates)" If="'$(DDCVerifyCustomBucketRates)' != ''"/>
 	<Property Name="DDCVerifyArgs" Value="$(DDCVerifyArgs) -DDC-VerifyRate=$(DDCVerifyDefaultRate)"/>
