Horde: Add a test job to generate arbitrary UGS artifacts.

#rnx

diff --git a/Engine/Build/Graph/Examples/UgsArtifact.xml b/Engine/Build/Graph/Examples/UgsArtifact.xml
new file mode 100644
--- /dev/null
+++ b/Engine/Build/Graph/Examples/UgsArtifact.xml
@@ -0,0 +1,19 @@
+<?xml version='1.0' ?>
+<BuildGraph xmlns="http://www.epicgames.com/BuildGraph" xmlns:xsi="http://www.w3.org/2001/XMLSchema-instance" xsi:schemaLocation="http://www.epicgames.com/BuildGraph ../Schema.xsd" >
+
+	<!-- 
+	This graph uploads an artifact to Horde which is exposed to the 'Sync Precompiled Binaries' menu in UGS. Users can choose to download this artifact type along with their regular sync.
+	-->
+
+	<Option Name="UgsProject" Restrict="//.*" DefaultValue="//UE5/Main/EngineTest/EngineTest.uproject" Description="Identifier for the project in UGS (typically a depot path)" />
+	<Option Name="UgsType" Restrict="[a-zA-Z]+" DefaultValue="Test" Description="Type of the archive for UGS. Zero or one of each archive type can be selected at a time in the UGS UI." />
+
+	<Agent Name="Agent 1" Type="Win64">
+		<Node Name="Create Artifact" Produces="#ArtifactFiles">
+			<WriteTextFile File="$(UgsType).txt" Text="This is an '$(UgsType)' artifact"/>
+		</Node>
+	</Agent>
+
+	<Artifact Name="$(UgsType)" Description="Type: $(UgsType)" Type="ugs-pcb" Keys="ugs-project=$(UgsProject)" Metadata="ArchiveType=$(UgsType)" Tag="#ArtifactFiles"/>
+
+</BuildGraph>
