Tag the output file for UGS artifact test.

#rnx

diff --git a/Engine/Build/Graph/Examples/UgsArtifact.xml b/Engine/Build/Graph/Examples/UgsArtifact.xml
--- a/Engine/Build/Graph/Examples/UgsArtifact.xml
+++ b/Engine/Build/Graph/Examples/UgsArtifact.xml
@@ -10,7 +10,7 @@
 
 	<Agent Name="Agent 1" Type="Win64">
 		<Node Name="Create Artifact" Produces="#ArtifactFiles">
-			<WriteTextFile File="$(UgsType).txt" Text="This is an '$(UgsType)' artifact"/>
+			<WriteTextFile File="$(UgsType).txt" Text="This is an '$(UgsType)' artifact" Tag="#ArtifactFiles"/>
 		</Node>
 	</Agent>
 
