Add optional snapshot generation and long term storage as part of the generic BuildAndTestProject build script.

diff --git a/Engine/Build/Graph/Tasks/BuildAndTestProject.xml b/Engine/Build/Graph/Tasks/BuildAndTestProject.xml
--- a/Engine/Build/Graph/Tasks/BuildAndTestProject.xml
+++ b/Engine/Build/Graph/Tasks/BuildAndTestProject.xml
@@ -1,6 +1,9 @@
 <?xml version='1.0' ?>
 <BuildGraph xmlns="http://www.epicgames.com/BuildGraph" xmlns:xsi="http://www.w3.org/2001/XMLSchema-instance" xsi:schemaLocation="http://www.epicgames.com/BuildGraph ../Schema.xsd" >
 	<EnvVar Name="UE_HORDE_JOBID"/>
+	<EnvVar Name="UE_HORDE_STEPID"/>
+	<EnvVar Name="UE_HORDE_TEMPLATEID"/>
+	<EnvVar Name="UE_HORDE_LAST_WARNING_CL"/>
 	
 	<EnvVar Name="UE-CloudPublishHost"/>
 	<EnvVar Name="UE-CloudPublishHttpVersion"/>
@@ -71,6 +74,7 @@
 	<Option Name="SymbolServerPath" DefaultValue="" Description="Symbol server path"/>
 
 	<Option Name="PublishCloudSnapshot" DefaultValue="false" Restrict="true|false" Description="Publish build to cloud host" />
+	<Option Name="PublishZenSnapshot" DefaultValue="false" Restrict="true|false" Description="Publish build to zen host" />
 
 	<Option Name="EnableSymbolStore" DefaultValue="$(IsBuildMachine)" Restrict="true|false" Description="Enables the upload of debugging symbols" />
 	<Property Name="SymbolServerPath" Value="" If="$(EnableSymbolStore) == 'false'"/>
@@ -563,25 +567,42 @@
 				</Do>
 			</Node>
 
-			<!-- (Optional) Publish to Cloud -->
+			<!-- (Optional) Publish from local Zen to Cloud or remote Zen server -->
 			<Property Name="PlatformPublishStageDependencyName" Value="$(PlatformHashNodeName)" />
-			<Do If="$(PublishBuilds) and $(PublishCloudSnapshot)">
-				<Switch>
-					<Case If="'$(UE-ZenPublishHost)' != ''">
-						<Property Name="PlatformSnapshotPublishNodeName" Value="$(PreNodeName)Snapshot Publish $(TargetPlatform)" />
-						<Node Name="$(PlatformSnapshotPublishNodeName)" Requires="$(PlatformHashNodeName)">
-							<ZenExportSnapshot Project="$(TargetProject)" Platform="$(CookPlatform)" DestinationStorageType="Zen" DestinationZenHost="$(UE-ZenPublishHost)" DestinationIdentifier="$(EscapedBranch)/$(BuildNamePath)" SnapshotDescriptorFile="$(NetworkOutputDirectory)/$(StagedPlatformFolder)/SnapshotDescriptor.json"/>
-						</Node>
-						<Property Name="PlatformPublishStageDependencyName" Value="$(PlatformSnapshotPublishNodeName)" />
-					</Case>
-					<Case If="'$(UE-CloudPublishHost)' != '' and '$(UE-CloudPublishNamespace)' != ''">
-						<Property Name="PlatformSnapshotPublishNodeName" Value="$(PreNodeName)Snapshot Publish $(TargetPlatform)" />
-						<Node Name="$(PlatformSnapshotPublishNodeName)" Requires="$(PlatformHashNodeName)">
-							<ZenExportSnapshot Project="$(TargetProject)" Platform="$(CookPlatform)" DestinationStorageType="Cloud" DestinationCloudHost="$(UE-CloudPublishHost)" DestinationCloudHttpVersion="$(UE-CloudPublishHttpVersion)" SnapshotDescriptorCloudHost="$(UE-CloudPublishDescriptorHost)" SnapshotDescriptorCloudHttpVersion="$(UE-CloudPublishDescriptorHttpVersion)" DestinationCloudNamespace="$(UE-CloudPublishNamespace)" DestinationIdentifier="$(EscapedBranch)/$(BuildNamePath)" SnapshotDescriptorFile="$(NetworkOutputDirectory)/$(StagedPlatformFolder)/SnapshotDescriptor.json"/>
-						</Node>
-						<Property Name="PlatformPublishStageDependencyName" Value="$(PlatformSnapshotPublishNodeName)" />
-					</Case>
-				</Switch>
+			<Property Name="ShouldSnapshotPublishCloud" Value="false"/>
+			<Property Name="ShouldSnapshotPublishCloud" Value="true" If="$(PublishCloudSnapshot) and '$(UE-CloudPublishHost)' != '' and '$(UE-CloudPublishNamespace)' != ''"/>
+			<Property Name="ShouldSnapshotPublishZen" Value="false"/>
+			<Property Name="ShouldSnapshotPublishZen" Value="true" If="$(IsBuildMachine) and !$(IsPreflight) and $(PublishZenSnapshot) and '$(UE-ZenPublishHost)' != ''"/>
+			<Property Name="ShouldSnapshotPublish" Value="false"/>
+			<Property Name="ShouldSnapshotPublish" Value="true" If="$(IsBuildMachine) and !$(IsPreflight) and $(ShouldSnapshotPublishCloud) or $(ShouldSnapshotPublishZen)"/>
+			<Do If="$(PublishBuilds) and $(ShouldSnapshotPublish)">
+				<Property Name="PlatformSnapshotPublishNodeName" Value="$(PreNodeName)Snapshot Publish $(TargetPlatform)" />
+				<Node Name="$(PlatformSnapshotPublishNodeName)" Requires="$(PlatformHashNodeName)" Produces="#SnapshotDescriptor_$(StagedPlatformFolder)">
+					<Property Name="SnapshotLocalDir" Value="$(RootDir)/Snapshots" />
+					<Property Name="SnapshotIdentifier" Value="$(EscapedBranch)/$(BuildNamePath)"/>
+					<Property Name="SnapshotIdentifier" Value="$(UE_HORDE_TEMPLATEID).$(BuildVersion).$(UE_HORDE_JOBID:-6).$(UE_HORDE_STEPID)" If="'$(UE_HORDE_JOBID)' != '' and '$(UE_HORDE_STEPID)' != '' and '$(UE_HORDE_TEMPLATEID)' != ''"/>
+					<Property Name="SnapshotFilenamePrefix" Value="SnapshotDescriptor"/>
+					<Property Name="SnapshotFilenamePrefix" Value="$(UE_HORDE_TEMPLATEID)" If="'$(UE_HORDE_TEMPLATEID)' != ''"/>
+					<Property Name="SnapshotLongTermStorageDir" Value="$(NetworkSnapshotDescFileRoot)/$(EscapedBranch)/$(StagedPlatformFolder)/$(Change)"  If="'$(NetworkSnapshotDescFileRoot)' != ''"/>
+					<Do If="$(ShouldSnapshotPublishCloud)">
+						<Property Name="SnapshotBaseDescriptorFile" Value="$(NetworkSnapshotDescFileRoot)/$(EscapedBranch)/$(StagedPlatformFolder)/$(UE_HORDE_LAST_WARNING_CL)/$(StagedPlatformFolder)/$(SnapshotFilenamePrefix)-cloud.json"  If="'$(NetworkSnapshotDescFileRoot)' != '' and '$(UE_HORDE_LAST_WARNING_CL)' != ''"/>
+						<ZenExportSnapshot Project="$(TargetProject)" Platform="$(CookPlatform)" DestinationStorageType="Cloud" DestinationCloudHost="$(UE-CloudPublishHost)" DestinationCloudHttpVersion="$(UE-CloudPublishHttpVersion)" SnapshotDescriptorCloudHost="$(UE-CloudPublishDescriptorHost)" SnapshotDescriptorCloudHttpVersion="$(UE-CloudPublishDescriptorHttpVersion)" DestinationCloudNamespace="$(UE-CloudPublishNamespace)" DestinationIdentifier="$(SnapshotIdentifier)" SnapshotDescriptorFile="$(SnapshotLocalDir)/$(StagedPlatformFolder)/$(SnapshotFilenamePrefix)-cloud.json" SnapshotBaseDescriptorFile="$(SnapshotBaseDescriptorFile)"/>
+						<Copy From="$(SnapshotLocalDir)/$(StagedPlatformFolder)/$(SnapshotFilenamePrefix)-cloud.json" To="$(NetworkOutputDirectory)/$(StagedPlatformFolder)/SnapshotDescriptor.json"/>
+						<Do If="'$(SnapshotLongTermStorageDir)' != ''">
+							<Copy From="$(SnapshotLocalDir)/$(StagedPlatformFolder)/$(SnapshotFilenamePrefix)-cloud.json" To="$(SnapshotLongTermStorageDir)/$(SnapshotFilenamePrefix)-cloud.json"/>
+						</Do>
+					</Do>
+					<Do If="$(ShouldSnapshotPublishZen)">
+						<ZenExportSnapshot Project="$(TargetProject)" Platform="$(CookPlatform)" DestinationStorageType="Zen" DestinationZenHost="$(UE-ZenPublishHost)" DestinationIdentifier="$(SnapshotIdentifier)" SnapshotDescriptorFile="$(SnapshotLocalDir)/$(StagedPlatformFolder)/$(SnapshotFilenamePrefix)-zen.json"/>
+						<Copy From="$(SnapshotLocalDir)/$(StagedPlatformFolder)/$(SnapshotFilenamePrefix)-zen.json" To="$(NetworkOutputDirectory)/$(StagedPlatformFolder)/SnapshotDescriptorZen.json"/>
+						<Do If="'$(SnapshotLongTermStorageDir)' != ''">
+							<Copy From="$(SnapshotLocalDir)/$(StagedPlatformFolder)/$(SnapshotFilenamePrefix)-zen.json" To="$(SnapshotLongTermStorageDir)/$(SnapshotFilenamePrefix)-zen.json"/>
+						</Do>
+					</Do>
+					<Tag Files="$(SnapshotLocalDir)/..." With="#SnapshotDescriptor_$(StagedPlatformFolder)"/>
+				</Node>
+				<Artifact Name="ci-snapshot-$(StagedPlatformFolder)" Description="Snapshot ($(StagedPlatformFolder))" Type="ugs-snapshot" BasePath="Snapshots" Keys="ugs-project=$(TargetProject)" Metadata="ArchiveType=SnapshotDescriptor_$(StagedPlatformFolder)" Tag="#SnapshotDescriptor_$(StagedPlatformFolder)"/>
+				<Property Name="PlatformPublishStageDependencyName" Value="$(PlatformSnapshotPublishNodeName)" />
 			</Do>
 			
 			<!-- Publish (Staged) -->
diff --git a/Engine/Build/Graph/Tasks/Inc/SharedProperties.xml b/Engine/Build/Graph/Tasks/Inc/SharedProperties.xml
--- a/Engine/Build/Graph/Tasks/Inc/SharedProperties.xml
+++ b/Engine/Build/Graph/Tasks/Inc/SharedProperties.xml
@@ -14,6 +14,8 @@
 	<Option Name="GauntletDeviceService" DefaultValue="" Description="Gauntlet Device Service to use for device reservation when running tests"/>
 	<!-- Android adbkeys -->
 	<Option Name="AdbKeysPath" DefaultValue="" Description="Path to adbkeys location (for Android)"/>
+	<!-- Network path that can be used for storing snapshot descriptors -->
+	<Property Name="NetworkSnapshotDescFileRoot" Value=""/>
 	<!-- Get custom shared properties -->
 	<Do If="Exists('$(RootDir)/$(SharedPropertiesFile)')">
 		<Include Script="$(RootDir)/$(SharedPropertiesFile)"/>
