Add dashboard category for snapshot descriptors to group them in the UI.

diff --git a/Engine/Build/Graph/Tasks/BuildAndTestProject.xml b/Engine/Build/Graph/Tasks/BuildAndTestProject.xml
--- a/Engine/Build/Graph/Tasks/BuildAndTestProject.xml
+++ b/Engine/Build/Graph/Tasks/BuildAndTestProject.xml
@@ -601,7 +601,7 @@
 					</Do>
 					<Tag Files="$(SnapshotLocalDir)/..." With="#SnapshotDescriptor_$(StagedPlatformFolder)"/>
 				</Node>
-				<Artifact Name="ci-snapshot-$(StagedPlatformFolder)" Description="Snapshot ($(StagedPlatformFolder))" Type="ugs-snapshot" BasePath="Snapshots" Keys="ugs-project=$(TargetProject)" Metadata="ArchiveType=SnapshotDescriptor_$(StagedPlatformFolder)" Tag="#SnapshotDescriptor_$(StagedPlatformFolder)"/>
+				<Artifact Name="ci-snapshot-$(StagedPlatformFolder)" Description="Snapshot ($(StagedPlatformFolder))" Type="ugs-snapshot" BasePath="Snapshots" Keys="ugs-project=$(TargetProject)" Metadata="ArchiveType=SnapshotDescriptor_$(StagedPlatformFolder);dashboard-category=Snapshots" Tag="#SnapshotDescriptor_$(StagedPlatformFolder)"/>
 				<Property Name="PlatformPublishStageDependencyName" Value="$(PlatformSnapshotPublishNodeName)" />
 			</Do>
 			
