Harmonize Options input between test automation buildgraphes
+ Remove reference to internal server in public buildgraph
+ Allow disabling Horde test report generation

#rnx

diff --git a/Engine/Build/Graph/Tasks/BuildAndTestProject.xml b/Engine/Build/Graph/Tasks/BuildAndTestProject.xml
--- a/Engine/Build/Graph/Tasks/BuildAndTestProject.xml
+++ b/Engine/Build/Graph/Tasks/BuildAndTestProject.xml
@@ -12,7 +12,9 @@
 	<EnvVar Name="UE-CloudPublishNamespace"/>
 	<EnvVar Name="UE-ZenPublishHost"/>
 
-	<Option Name="Horde" Restrict="true|false" DefaultValue="$(IsBuildMachine)" Description="Whether the build is running under Horde"/>
+	<Property Name="IsHordeEnv" Value="false"/>
+	<Property Name="IsHordeEnv" Value="true" If="'$(UE_HORDE_JOBID)' != ''"/>
+	<Option Name="Horde" DefaultValue="$(IsBuildMachine) And $(IsHordeEnv)" Description="Whether the build is running under Horde"/>
 	
 	<!-- Project Options -->
 	<Option Name="ProjectFile" DefaultValue="" Restrict=".*" Description="Path to the project file. Can be used instead of specifying ProjectName and ProjectPath separately."/>
diff --git a/Engine/Build/Graph/Tasks/Inc/GauntletSettings.xml b/Engine/Build/Graph/Tasks/Inc/GauntletSettings.xml
--- a/Engine/Build/Graph/Tasks/Inc/GauntletSettings.xml
+++ b/Engine/Build/Graph/Tasks/Inc/GauntletSettings.xml
@@ -28,8 +28,10 @@
 
 	<Option Name="AdbKeysPath" DefaultValue="" Description="Path to adbkeys location (for Android)"/>
 
+	<Option Name="DisableTestResultsForHorde" DefaultValue="false" Restrict="true|false" Description="Force disabling test results for Horde production"/>
+
 	<Property Name="GauntletCommonArgs" Value="-maxduration=$(GauntletTimeOut) -unattended -verbose -ResumeOnCriticalFailure $(ExtraGauntletCommonArgs)" />
-	<Property Name="GauntletCommonArgs" Value="$(GauntletCommonArgs) -WriteTestResultsForHorde" If="'$(UE_HORDE_JOBID)' != ''"/>
+	<Property Name="GauntletCommonArgs" Value="$(GauntletCommonArgs) -WriteTestResultsForHorde" If="$(Horde) and !$(DisableTestResultsForHorde)"/>
 	<Property Name="GauntletCommonArgs" Value="$(GauntletCommonArgs) -branch=&quot;$(Branch)&quot; -changelist=&quot;$(Change)&quot;" />
 
 	<Property Name="GauntletDeviceArgs" Value="" />
diff --git a/Engine/Build/Graph/Tasks/Inc/PathSettings.xml b/Engine/Build/Graph/Tasks/Inc/PathSettings.xml
--- a/Engine/Build/Graph/Tasks/Inc/PathSettings.xml
+++ b/Engine/Build/Graph/Tasks/Inc/PathSettings.xml
@@ -22,7 +22,7 @@
 
 		OutOption: PublishBuilds - Publish output from builds to a network share
 		OutOption: PublishHashedStagingData - Publish staged content with hashes for quick syncing
-		OutOption: PublishReports - Publish reports from the build
+		OutOption: PublishReports - Publish automation reports to the network
 
 		OutProperty: ProjectOutputDirectory - Local path to package builds to
 		OutProperty: ShaderSymbolsDirectory - Local path to write shader symbols to
@@ -43,12 +43,15 @@
 	<!-- Publish a hashed staging data that developers can quickly sync to -->
 	<Option Name="PublishHashedStagingData" DefaultValue="false" Restrict="true|false" Description="Create hashes for syncing via Unsync" />
 
-	<!-- Publish automation reports to the network-->
-	<Option Name="PublishReports" DefaultValue="$(IsEpicBuildMachine)" Description="Publish packaged builds to the builds folder on the netwok"/>
+	<!-- Publish automation reports to the network -->
+	<Option Name="PublishReports" DefaultValue="$(IsEpicBuildMachine)" Description="Publish automation reports to the network"/>
 
 	<!-- For pre-generated build -->
 	<Option Name="GeneratedBuildVersion" DefaultValue="" Description="The pre-generated build version"/>
 
+	<!-- Copy test report to a network share -->
+	<Option Name="CopyTestReportOnNetShare" DefaultValue="$(IsBuildMachine)" Description="Copy test report to a network share"/>
+
 	<!-- Where builds are saved to locally -->
 	<Property Name="ProjectOutputDirectory" Value="$(ProjectPath)/Saved/StagedBuilds"/>
 	<Property Name="ShaderSymbolsDirectory" Value="$(ProjectPath)/Saved/ShaderSymbols"/>
@@ -79,7 +82,7 @@
 	<Do If="$(IsEpicBuildMachine) or $(PublishReports)" >
 		<Property Name="AutomationReportOutputDirectory" Value="$(NetworkReportRoot)"/>
 		<Property Name="AutomationReportOutputDirectory" Value="$(AutomationReportOutputDirectory)/$(EscapedBranch)/$(ProjectName)/$(CL-String)"/>
-		<Property Name="AutomationReportUrl" Value="http://automation.epicgames.net/reports/$(EscapedBranch)/$(ProjectName)/$(CL-String)"/>
+		<Property Name="AutomationReportUrl" Value="$(AutomationReportURLBase)/$(EscapedBranch)/$(ProjectName)/$(CL-String)"/>
 	</Do>
 
 	<!-- Network path that can be used for storing reports and logs -->
diff --git a/Engine/Build/Graph/Tasks/Inc/SharedProperties.xml b/Engine/Build/Graph/Tasks/Inc/SharedProperties.xml
--- a/Engine/Build/Graph/Tasks/Inc/SharedProperties.xml
+++ b/Engine/Build/Graph/Tasks/Inc/SharedProperties.xml
@@ -16,6 +16,8 @@
 	<Option Name="AdbKeysPath" DefaultValue="" Description="Path to adbkeys location (for Android)"/>
 	<!-- Network path that can be used for storing snapshot descriptors -->
 	<Property Name="NetworkSnapshotDescFileRoot" Value=""/>
+	<!-- Automation report website -->
+	<Property Name="AutomationReportURLBase" Value=""/>
 	<!-- Get custom shared properties -->
 	<Do If="Exists('$(RootDir)/$(SharedPropertiesFile)')">
 		<Include Script="$(RootDir)/$(SharedPropertiesFile)"/>
