Make agent type configurable

diff --git a/Engine/Build/Graph/Tasks/DeviceOperation.xml b/Engine/Build/Graph/Tasks/DeviceOperation.xml
--- a/Engine/Build/Graph/Tasks/DeviceOperation.xml
+++ b/Engine/Build/Graph/Tasks/DeviceOperation.xml
@@ -6,6 +6,8 @@
 	<Option Name="InstallBuildFlags" DefaultValue="" Description="Optional BuildFlags used to specific a type of build to install."/>
 	<Option Name="ProjectName" DefaultValue="" Description="The Project Name of the build."/>
 
+	<Option Name="InstallBuildAgent" DefaultValue="DevkitAutomation" Description="Name of agent pool"/>
+	
 	<Option Name="Configurations" DefaultValue="" Description="The Configurations of the build."/>
 	<Option Name="WithShipping" DefaultValue="false" Restrict="true|false" Description="With Shipping Configuration."/>
 	<Property Name="Configurations" Value="$(Configurations);Shipping" If="$(WithShipping)"/>
@@ -35,7 +37,7 @@
 	<Option Name="DeviceList" DefaultValue="" Description="List of device to use."/>
 	<Option Name="ExtraInstallArgs" DefaultValue="" Description="Additional install arguments."/>
 	
-	<Agent Name="Install Build Job" Type="DevkitAutomation">
+	<Agent Name="Install Build Job" Type="$(InstallBuildAgent)">
 		<ForEach Name="TargetPlatform" Values="$(TargetRemotePlatforms)">
 			<Node Name="InstallBuild $(TargetPlatform)">
 			<ForEach Name="Configuration" Values="$(Configurations)">
