Reconfigure the DDC-Verify job.

diff --git a/Engine/Build/Graph/Tests/DDCVerify.xml b/Engine/Build/Graph/Tests/DDCVerify.xml
--- a/Engine/Build/Graph/Tests/DDCVerify.xml
+++ b/Engine/Build/Graph/Tests/DDCVerify.xml
@@ -5,6 +5,7 @@
 	<Option Name="ProjectFile" DefaultValue="_Unknown_" Restrict=".+" Description="The projects to cook"/>
 	<Option Name="ProjectName" DefaultValue="_Unknown_" Restrict=".+" Description="The projects to cook"/>
 	<Option Name="EditorTarget" DefaultValue="$(ProjectName)Editor" Description="Compile name of the editor target"/>
+	<Option Name="DDCGraph" DefaultValue="" Description="Overridable DDC graph to use when performing the verification (default: empty, uses default graph)"/>
 
 	<!-- Note we can not make the default='', because the default is used for some validation. -->
 	<Error Message="No project file specified" If="'$(ProjectFile)' == '_Unknown_'"/>
@@ -28,17 +29,25 @@
 	<Property Name="DDCVerifyArgs" Value="$(DDCVerifyArgs)=$(DDCVerifyCustomBucketRates)" If="'$(DDCVerifyCustomBucketRates)' != ''"/>
 	<Property Name="DDCVerifyArgs" Value="$(DDCVerifyArgs) -DDC-VerifyRate=$(DDCVerifyDefaultRate)"/>
 
+	<Property Name="DDCGraphArgs" Value=""/>
+	<Property Name="DDCGraphArgs" Value="-ddc=$(DDCGraph)" If="'$(DDCGraph)' != ''"/>
+
 	<!--========================================-->
 	<Property Name="EditorPlatform" Value="Win64"/>
-	<Property Name="CommonCookArgs" Value="-targetplatform=$(TargetPlatform) $(DDCVerifyArgs) -skiponlyeditoronly=true"/>
+	<Property Name="CommonCookArgs" Value="-targetplatform=$(TargetPlatform) $(DDCVerifyArgs) $(DDCGraphArgs) -skiponlyeditoronly=true"/>
 
 	<Property Name="CompileNodeName" Value="Compile $(EditorPlatform) $(ProjectName)"/>
 
 	<Agent Name="DDCVerifyAgent" Type="$(AgentType)">
-		<Node Name="$(CompileNodeName)">
+		<Node Name="Update Version Files">
+			<SetVersion Change="$(Change)" Branch="$(EscapedBranch)" CompatibleChange="$(CodeChange)" If="$(IsBuildMachine)"/>
+		</Node>
+
+		<Node Name="$(CompileNodeName)" Requires="Update Version Files">
 			<Compile Target="ShaderCompileWorker" Platform="$(EditorPlatform)" Project="$(ProjectFile)" Configuration="Development" Arguments="$(UBACompileArgs)"/>
 			<Compile Target="$(EditorTarget)" Platform="$(EditorPlatform)" Project="$(ProjectFile)" Configuration="Development" Arguments="$(UBACompileArgs)"/>
 		</Node>
+
 		<Node Name="DDC Verify" Requires="$(CompileNodeName)">
 			<Commandlet Name="Cook" Project="$(ProjectFile)" Arguments="$(CommonCookArgs)"/>
 		</Node>
