Horde: Add metadata to artifacts produced by BuildAndTestProject.xml.

#rnx

diff --git a/Engine/Build/Graph/Tasks/BuildAndTestProject.xml b/Engine/Build/Graph/Tasks/BuildAndTestProject.xml
--- a/Engine/Build/Graph/Tasks/BuildAndTestProject.xml
+++ b/Engine/Build/Graph/Tasks/BuildAndTestProject.xml
@@ -624,7 +624,7 @@
 			<Node Name="$(PlatformStagedArtifactNodeName)" Requires="$(PlatformStageNodeName)" Produces="#$(PlatformStageNodeName) Files">
 				<Tag Files="$(ProjectOutputDirectory)/$(StagedPlatformFolder)/..." With="#$(PlatformStageNodeName) Files"/>
 			</Node>
-			<Artifact Name="$(PlatformPackageArtifactName)" Type="packaged-build" BasePath="$(ProjectOutputDirectory)/$(StagedPlatformFolder)" Tag="#$(PlatformStageNodeName) Files"/>
+			<Artifact Name="$(PlatformPackageArtifactName)" Type="packaged-build" BasePath="$(ProjectOutputDirectory)/$(StagedPlatformFolder)" Keys="Project=$(ProjectName)" Metadata="Stream=$(Branch);Platform=$(TargetPlatform)" Tag="#$(PlatformStageNodeName) Files"/>
 
 			<!-- Package -->
 			<Node Name="$(PlatformPackageNodeName)" Requires="$(PlatformPublishStageDependencyName)">
