Add Horde Artifacts for Package Builds, update previously existing Staged Build Artifacts with explicit labels

#rnx

diff --git a/Engine/Build/Graph/Tasks/BuildAndTestProject.xml b/Engine/Build/Graph/Tasks/BuildAndTestProject.xml
--- a/Engine/Build/Graph/Tasks/BuildAndTestProject.xml
+++ b/Engine/Build/Graph/Tasks/BuildAndTestProject.xml
@@ -475,10 +475,13 @@
 		<Property Name="PlatformPublishPackagedNodeName" Value="$(PreNodeName)Publish Packaged $(TargetPlatform)" />
 		<Property Name="PlatformCheckBuildSizeNodeName" Value="$(PreNodeName)Check Build Size $(TargetPlatform)" />
 		<Property Name="PlatformStagedArtifactNodeName" Value="$(PlatformPublishStagedNodeName) (Horde)" />
-		<Property Name="RequirePlatformStaged" Value="$(PlatformStageNodeName);$(PlatformStagedArtifactNodeName)" />
+		<Property Name="HordeArtifactPlatformPackagedNodeName" Value="$(PlatformPublishPackagedNodeName) (Horde)" />
+		<Property Name="RequirePlatformStaged" Value="$(PlatformStageNodeName)" />
+		<Property Name="RequirePlatformStaged" Value="$(RequirePlatformStaged);$(PlatformStagedArtifactNodeName)" If="$(Horde)"/>
 		<Property Name="RequirePlatformStaged" Value="$(RequirePlatformStaged);$(PlatformPublishStagedNodeName)" If="!$(SkipPublish) And $(PublishBuilds)"/>
 		<Property Name="RequirePlatformPackaged" Value="$(PlatformPackageNodeName)" />
-		<Property Name="RequirePlatformPackaged" Value="$(PlatformPublishPackagedNodeName)"  If="!$(SkipPublish) And $(PublishBuilds)"/>
+		<Property Name="RequirePlatformPackaged" Value="$(RequirePlatformPackaged);$(HordeArtifactPlatformPackagedNodeName)" If="$(Horde)"/>
+		<Property Name="RequirePlatformPackaged" Value="$(RequirePlatformPackaged);$(PlatformPublishPackagedNodeName)"  If="!$(SkipPublish) And $(PublishBuilds)"/>
 		
 		<!-- Set CookPlatform since it isn't always the same as the target -->
 		<Property Name="CookPlatform" Value="$(TargetPlatform)" />
@@ -624,7 +627,7 @@
 			<Node Name="$(PlatformStagedArtifactNodeName)" Requires="$(PlatformStageNodeName)" Produces="#$(PlatformStageNodeName) Files">
 				<Tag Files="$(ProjectOutputDirectory)/$(StagedPlatformFolder)/..." With="#$(PlatformStageNodeName) Files"/>
 			</Node>
-			<Artifact Name="$(PlatformPackageArtifactName)" Type="packaged-build" BasePath="$(ProjectOutputDirectory)/$(StagedPlatformFolder)" Keys="Project=$(ProjectName)" Metadata="Stream=$(Branch);Platform=$(TargetPlatform)" Tag="#$(PlatformStageNodeName) Files"/>
+			<Artifact Name="$(PlatformPackageArtifactName) (Staged)" Type="staged-build" BasePath="$(ProjectOutputDirectory)/$(StagedPlatformFolder)" Keys="Project=$(ProjectName)" Metadata="Stream=$(Branch);Platform=$(StagedPlatformFolder);BuildName=$(BuildNamePath);dashboard-category=Staged Builds" Tag="#$(PlatformStageNodeName) Files"/>
 
 			<!-- Package -->
 			<Node Name="$(PlatformPackageNodeName)" Requires="$(PlatformPublishStageDependencyName)">
@@ -640,6 +643,21 @@
 					<Command Name="BuildCookRun" Arguments="$(BCRArgs) -configuration=$(PackageConfiguration) $(StageAndPackageArguments) -skipbuild -skipcook -skipstage -skippak -skippackage -archive -archivedirectory=&quot;$(NetworkOutputDirectory)/$(StagedPlatformFolder)/Packages&quot;" />
 				</ForEach>
 			</Node>
+			
+			<!-- Publish (Packages) (Horde) -->
+			<Property Name="PackagedArtifactTags" Value=""/>
+			<ForEach Name="PackageConfiguration" Values="$(PackageConfigurations)" Separator="+">
+				<Property Name="PackagedArtifactTags" Value="#$(PlatformPackageNodeName) $(PackageConfiguration) Files;$(PackagedArtifactTags)"/>
+			</ForEach>
+			<Node Name="$(HordeArtifactPlatformPackagedNodeName)" Requires="$(PlatformPackageNodeName)" Produces="$(PackagedArtifactTags)">
+				<ForEach Name="PackageConfiguration" Values="$(PackageConfigurations)" Separator="+">
+					<Command Name="BuildCookRun" Arguments="$(BCRArgs) -configuration=$(PackageConfiguration) $(StageAndPackageArguments) -skipbuild -skipcook -skipstage -skippak -skippackage -archive -archivedirectory=&quot;$(ProjectOutputDirectory)/$(StagedPlatformFolder)/Packages/$(PackageConfiguration)&quot;" />
+					<Tag Files="$(ProjectOutputDirectory)/$(StagedPlatformFolder)/Packages/$(PackageConfiguration)/..." With="#$(PlatformPackageNodeName) $(PackageConfiguration) Files"/>
+				</ForEach>
+			</Node>
+			<ForEach Name="PackageConfiguration" Values="$(PackageConfigurations)" Separator="+">
+				<Artifact Name="$(PlatformPackageArtifactName) ($(PackageConfiguration))" Type="packaged-bulid" BasePath="$(ProjectOutputDirectory)/$(StagedPlatformFolder)/Packages/$(PackageConfiguration)" Keys="Project=$(ProjectName);BuildName=$(BuildNamePath)" Metadata="Stream=$(Branch);Platform=$(StagedPlatformFolder);Configuration=$(PackageConfiguration);dashboard-category=$(PackageConfiguration) Builds" Tag="#$(PlatformPackageNodeName) $(PackageConfiguration) Files"/>
+			</ForEach>			
 
 			<!-- Check Build Size -->
 			<Node Name="$(PlatformCheckBuildSizeNodeName)" Requires="$(PlatformPublishPackagedNodeName)">
