Fix typo `packaged-bulid` -> `packaged-build`

diff --git a/Engine/Build/Graph/Tasks/BuildAndTestProject.xml b/Engine/Build/Graph/Tasks/BuildAndTestProject.xml
--- a/Engine/Build/Graph/Tasks/BuildAndTestProject.xml
+++ b/Engine/Build/Graph/Tasks/BuildAndTestProject.xml
@@ -656,7 +656,7 @@
 				</ForEach>
 			</Node>
 			<ForEach Name="PackageConfiguration" Values="$(PackageConfigurations)" Separator="+">
-				<Artifact Name="$(PlatformPackageArtifactName) ($(PackageConfiguration))" Type="packaged-bulid" BasePath="$(ProjectOutputDirectory)/$(StagedPlatformFolder)/Packages/$(PackageConfiguration)" Keys="Project=$(ProjectName);BuildName=$(BuildNamePath)" Metadata="Stream=$(Branch);Platform=$(StagedPlatformFolder);Configuration=$(PackageConfiguration);dashboard-category=$(PackageConfiguration) Builds" Tag="#$(PlatformPackageNodeName) $(PackageConfiguration) Files"/>
+				<Artifact Name="$(PlatformPackageArtifactName) ($(PackageConfiguration))" Type="packaged-build" BasePath="$(ProjectOutputDirectory)/$(StagedPlatformFolder)/Packages/$(PackageConfiguration)" Keys="Project=$(ProjectName);BuildName=$(BuildNamePath)" Metadata="Stream=$(Branch);Platform=$(StagedPlatformFolder);Configuration=$(PackageConfiguration);dashboard-category=$(PackageConfiguration) Builds" Tag="#$(PlatformPackageNodeName) $(PackageConfiguration) Files"/>
 			</ForEach>			
 
 			<!-- Check Build Size -->
