Update Horde Artifact output so that other configuration exes are not tagged in the final artifact.

#tests also downloaded the resulting package from the preflight to confirm it started

diff --git a/Engine/Build/Graph/Tasks/BuildAndTestProject.xml b/Engine/Build/Graph/Tasks/BuildAndTestProject.xml
--- a/Engine/Build/Graph/Tasks/BuildAndTestProject.xml
+++ b/Engine/Build/Graph/Tasks/BuildAndTestProject.xml
@@ -651,12 +651,12 @@
 			</ForEach>
 			<Node Name="$(HordeArtifactPlatformPackagedNodeName)" Requires="$(PlatformPackageNodeName)" Produces="$(PackagedArtifactTags)">
 				<ForEach Name="PackageConfiguration" Values="$(PackageConfigurations)" Separator="+">
-					<Command Name="BuildCookRun" Arguments="$(BCRArgs) -configuration=$(PackageConfiguration) $(StageAndPackageArguments) -skipbuild -skipcook -skipstage -skippak -skippackage -archive -archivedirectory=&quot;$(ProjectOutputDirectory)/$(StagedPlatformFolder)/Packages/$(PackageConfiguration)&quot;" />
-					<Tag Files="$(ProjectOutputDirectory)/$(StagedPlatformFolder)/Packages/$(PackageConfiguration)/..." With="#$(PlatformPackageNodeName) $(PackageConfiguration) Files"/>
+					<Command Name="BuildCookRun" Arguments="$(BCRArgs) -configuration=$(PackageConfiguration) $(StageAndPackageArguments) -skipbuild -skipcook -stage -skippak -package -archive -archivedirectory=&quot;$(ProjectOutputDirectory)/Packages/$(StagedPlatformFolder)/$(PackageConfiguration)&quot;" />
+					<Tag Files="$(ProjectOutputDirectory)/Packages/$(StagedPlatformFolder)/$(PackageConfiguration)/..." With="#$(PlatformPackageNodeName) $(PackageConfiguration) Files"/>
 				</ForEach>
 			</Node>
 			<ForEach Name="PackageConfiguration" Values="$(PackageConfigurations)" Separator="+">
-				<Artifact Name="$(PlatformPackageArtifactName) ($(PackageConfiguration))" Type="packaged-build" BasePath="$(ProjectOutputDirectory)/$(StagedPlatformFolder)/Packages/$(PackageConfiguration)" Keys="Project=$(ProjectName)" Metadata="Stream=$(Branch);Platform=$(StagedPlatformFolder);Configuration=$(PackageConfiguration);BuildName=$(BuildNamePath);dashboard-category=$(PackageConfiguration) Builds" Tag="#$(PlatformPackageNodeName) $(PackageConfiguration) Files"/>
+				<Artifact Name="$(PlatformPackageArtifactName) ($(PackageConfiguration))" Type="packaged-build" BasePath="$(ProjectOutputDirectory)/Packages/$(StagedPlatformFolder)/$(PackageConfiguration)" Keys="Project=$(ProjectName)" Metadata="Stream=$(Branch);Platform=$(StagedPlatformFolder);Configuration=$(PackageConfiguration);BuildName=$(BuildNamePath);dashboard-category=$(PackageConfiguration) Builds" Tag="#$(PlatformPackageNodeName) $(PackageConfiguration) Files"/>
 			</ForEach>			
 
 			<!-- Check Build Size -->
