Workaround platform specific issue with executable staging errors.
A previous CL limited the executables staged to the platform to one executable, but one platform already did this.

#rnx

#tests Other Preflight: 66a10bfa312522000da50874, Win64 & Mac packages downloaded and confirmed start: 66a10bfa312522000da50874

diff --git a/Engine/Build/Graph/Tasks/BuildAndTestProject.xml b/Engine/Build/Graph/Tasks/BuildAndTestProject.xml
--- a/Engine/Build/Graph/Tasks/BuildAndTestProject.xml
+++ b/Engine/Build/Graph/Tasks/BuildAndTestProject.xml
@@ -646,12 +646,15 @@
 			
 			<!-- Publish (Packages) (Horde) -->
 			<Property Name="PackagedArtifactTags" Value=""/>
+			<Property Name="HordePackageSkipBCRArguments" Value="-stage -package"/>
+			<Property Name="HordePackageSkipBCRArguments" Value="-skipstage -skippackage" If="'$(TargetPlatform)' == 'Mac'"/>
+
 			<ForEach Name="PackageConfiguration" Values="$(PackageConfigurations)" Separator="+">
 				<Property Name="PackagedArtifactTags" Value="#$(PlatformPackageNodeName) $(PackageConfiguration) Files;$(PackagedArtifactTags)"/>
 			</ForEach>
 			<Node Name="$(HordeArtifactPlatformPackagedNodeName)" Requires="$(PlatformPackageNodeName)" Produces="$(PackagedArtifactTags)">
 				<ForEach Name="PackageConfiguration" Values="$(PackageConfigurations)" Separator="+">
-					<Command Name="BuildCookRun" Arguments="$(BCRArgs) -configuration=$(PackageConfiguration) $(StageAndPackageArguments) -skipbuild -skipcook -stage -skippak -package -archive -archivedirectory=&quot;$(ProjectOutputDirectory)/Packages/$(StagedPlatformFolder)/$(PackageConfiguration)&quot;" />
+					<Command Name="BuildCookRun" Arguments="$(BCRArgs) -configuration=$(PackageConfiguration) $(StageAndPackageArguments) -skipbuild -skipcook -skippak $(HordePackageSkipBCRArguments) -archive -archivedirectory=&quot;$(ProjectOutputDirectory)/Packages/$(StagedPlatformFolder)/$(PackageConfiguration)&quot;" />
 					<Tag Files="$(ProjectOutputDirectory)/Packages/$(StagedPlatformFolder)/$(PackageConfiguration)/..." With="#$(PlatformPackageNodeName) $(PackageConfiguration) Files"/>
 				</ForEach>
 			</Node>
