Adding a macro to BuildAndTestProject which can be extended in order to inject additional test nodes attached to the existing per-platform/per-configuration test agent.

diff --git a/Engine/Build/Graph/Tasks/BuildAndTestProject.xml b/Engine/Build/Graph/Tasks/BuildAndTestProject.xml
--- a/Engine/Build/Graph/Tasks/BuildAndTestProject.xml
+++ b/Engine/Build/Graph/Tasks/BuildAndTestProject.xml
@@ -121,6 +121,15 @@
 	<Option Name="TargetPlatformsToSkipBootTests" DefaultValue="$(TargetPlatformsToSkipTests)" Description="Target platforms that should skip boot tests."/>
 	<Option Name="CookedEditorPlatformsToSkipBootTests" DefaultValue="$(CookedEditorPlatformsToSkipTests)" Description="Cooked editor platforms that should skip boot tests."/>
 
+	<!-- Macro that allows adding additional per-platform, per-configuration test nodes -->
+	<!-- Expand the macro inside the AdditionalTestNodesExtendFile --> 
+	<Option Name="AdditionalTestNodesExtendFile" DefaultValue="" Description="For extend the AdditionalTestNodes Macro" />
+	<Macro Name="AdditionalTestNodes" />
+	<Do If="'$(AdditionalTestNodesExtendFile)' != ''">
+		<Error Message="$(AdditionalTestNodesExtendFile) does not seem to exist" If="!Exists('$(AdditionalTestNodesExtendFile)')" />
+		<Include Script="$(AdditionalTestNodesExtendFile)" />
+	</Do>
+
 	<Option Name="TestsRequireEditor" DefaultValue="false" Description="Whether to force tests to require editor compilation"/>
 
 	<!-- Macros that allows custom pre-steps -->
@@ -767,9 +776,13 @@
 						<Property Name="AllPlatformTestNodes" Value="$(AllPlatformTestNodes);$(TestNodeName)"/>
 					</Do>
 				</ForEach>
+				
+				<!-- Expand the tests here -->
+				<Expand Name="AdditionalTestNodes"/>
+				
 			</ForEach>
 		</Agent>
-
+		
 		<!-- Declare labels for CIS -->
 		<Label Category="Clients" Name="$(PreLabelName)$(TargetPlatform)" Requires="$(PlatformCompileNodeName);$(UploadSymbolsNodeName)" />
 		<Label Category="Clients" Name="$(PreLabelName)Content $(TargetPlatform)" Requires="$(PlatformCookNodeName);$(PlatformStageNodeName);$(PlatformPackageNodeName)" />
