Change all test projects to use -cookincrmental as a cook argument when running on a build maching within Epic.
This ensures that the cooked output contains dependency information.  It also allows us to cook incrementally on build machines if an initial state is supplied.

diff --git a/Engine/Build/Graph/Tasks/BuildAndTestProject.xml b/Engine/Build/Graph/Tasks/BuildAndTestProject.xml
--- a/Engine/Build/Graph/Tasks/BuildAndTestProject.xml
+++ b/Engine/Build/Graph/Tasks/BuildAndTestProject.xml
@@ -189,6 +189,14 @@
 		<Property Name="ExtraTargetCompileArguments" Value="$(ExtraTargetCompileArguments) -nodebuginfo -WarningsAsErrors"/>
 	</Do>
 
+	<!-- On Epic build machines, we want to:
+		1. Write dependency information into the cooked output so that it is suitable for use as a base upon which to cook incrementally.
+		2. Cook incrementally if the initial state at the start of the job is not empty/blank.
+	-->
+	<Do If="$(IsEpicBuildMachine)">
+		<Property Name="ExtraCookArguments" Value="$(ExtraCookArguments) -cookincremental"/>
+	</Do>
+
 	<!-- Detect the type of the client based on the target name -->
 	<Property Name="TargetType" Value="Game" />
 	<Property Name="TargetType" Value="Client" If="Contains('$(TargetName)', 'Client')" />
