[MH Plugin] Copy test data from horde to shared drive

diff --git a/Engine/Build/Graph/Tasks/BuildAndTestProject.xml b/Engine/Build/Graph/Tasks/BuildAndTestProject.xml
--- a/Engine/Build/Graph/Tasks/BuildAndTestProject.xml
+++ b/Engine/Build/Graph/Tasks/BuildAndTestProject.xml
@@ -134,6 +134,14 @@
 
 	<!-- Macros that allows custom pre-steps -->
 	<Macro Name="PreCookStep"/>
+	
+	<!-- Macros that allow custom post tests steps -->
+	<Option Name="PostEditorTestsExtendFile" DefaultValue="" Description="Path to the file containing all post editor tests macro extends"/>
+	<Macro Name="PostEditorTestsStep"/>
+	<Do If="'$(PostEditorTestsExtendFile)' != ''">
+		<Error Message="Post Editor Tests Extend File $(PostEditorTestsExtendFile) does not exist" If="!Exists('$(PostEditorTestsExtendFile)')"/>
+		<Include Script="$(PostEditorTestsExtendFile)"/>
+	</Do>
 
 	<!-- Settings files that can be overriden -->
 	<Option Name="PathSettingsFile" DefaultValue="Inc/PathSettings.xml" Description="Path settings to use" />
@@ -330,6 +338,9 @@
 				<!-- Add each test node to our generated list -->
 				<Property Name="PlatformEditorTestNodes" Value="$(PlatformEditorTestNodes);$(TestNodeName)" />
 			</ForEach>
+			
+			<!-- Expand Post Editor Tests step -->
+			<Expand Name="PostEditorTestsStep"/>
 
 			<!-- Create Labels -->
 			<Label Category="Editors" Name="$(PreLabelName)Editor DDC $(EditorPlatform)" Requires="$(ProjectEditorDDC)" />
