Cooker:IterativeCompare Test: After cooking the forcerecook build, copy the non-package artifacts to a new directory, and export zen to that new directory, rather than exporting zen into the cooked output directory. We need the cooked output directory to be unmodified by the current job, so that it is in the correct state for the incremental cook at the beginning of the next job.
#rnx

diff --git a/Engine/Build/Graph/Tests/IterativeCompare.xml b/Engine/Build/Graph/Tests/IterativeCompare.xml
--- a/Engine/Build/Graph/Tests/IterativeCompare.xml
+++ b/Engine/Build/Graph/Tests/IterativeCompare.xml
@@ -29,7 +29,8 @@
 	<Property Name="CommonCookArgs" Value="-targetplatform=$(TargetPlatform) -zenstore -cookincremental"/>
 	<Property Name="IterativeCompareDescriptiveName" Value="IterativeCompare $(EditorPlatform) $(ProjectName)"/>
 	<Property Name="PlatformCookedDir" Value="$(FullPathProjectDir)/Saved/Cooked/$(TargetPlatform)"/>
-	<Property Name="CompareDir" Value="$(FullPathProjectDir)/Intermediate/IterativeCompare/$(TargetPlatform)"/>
+	<Property Name="CompareDirIncremental" Value="$(FullPathProjectDir)/Intermediate/IterativeCompare/Iter/$(TargetPlatform)"/>
+	<Property Name="CompareDirRecook" Value="$(FullPathProjectDir)/Intermediate/IterativeCompare/Recook/$(TargetPlatform)"/>
 	<Property Name="NewLine" Value="&#10;"/>
 
 	<Agent Name="IterativeCompareAgent" Type="$(AgentType)">
@@ -44,32 +45,40 @@
 			<Log Message="$(NewLine)========== IterativeCompare.xml: Incremental Cook =========="/>
 			<Commandlet Name="Cook" Project="$(ProjectFile)" Arguments="$(CommonCookArgs)"/>
 
-			<!-- Write the files from incremental cook's zenstore out into the Saved/Cooked/Platform dir. -->
-			<Log Message="$(NewLine)========== IterativeCompare.xml: ZenMirror Incremental Cook =========="/>
-			<Log Message="ZenMirror Project=$(ProjectFile) Platform=$(TargetPlatform)"/>
-			<ZenMirror Project="$(ProjectFile)" Platform="$(TargetPlatform)"/>
-
-			<!-- Move all cook artifacts to the compare dir. -->
+			<!-- Move all cook artifacts to the compare_incremental dir. -->
 			<Log Message="$(NewLine)========== IterativeCompare.xml: Move =========="/>
-			<Log Message="Delete Directories=$(CompareDir)"/>
-			<Delete Directories="$(CompareDir)"/>
-			<Log Message="Move From=$(PlatformCookedDir)/*.* To=$(CompareDir)"/>
-			<Move From="$(PlatformCookedDir)/*.*" To="$(CompareDir)" ErrorIfNotFound="true"/>
+			<Log Message="Delete Directories=$(CompareDirIncremental)"/>
+			<Delete Directories="$(CompareDirIncremental)"/>
+			<Log Message="Move From=$(PlatformCookedDir)/*.* To=$(CompareDirIncremental)"/>
+			<Move From="$(PlatformCookedDir)/*.*" To="$(CompareDirIncremental)" ErrorIfNotFound="true"/>
 			<Log Message="Delete Directories=$(PlatformCookedDir)"/>
 			<Delete Directories="$(PlatformCookedDir)"/>
 
+			<!-- Write the files from incremental cook's zenstore out into CompareDirIncremental. -->
+			<Log Message="$(NewLine)========== IterativeCompare.xml: ZenMirror Incremental Cook =========="/>
+			<Log Message="ZenMirror Project=$(ProjectFile) Platform=$(TargetPlatform) DestinationFileDir=$(CompareDirIncremental)"/>
+			<ZenMirror Project="$(ProjectFile)" Platform="$(TargetPlatform)" DestinationFileDir="$(CompareDirIncremental)"/>
+
 			<!-- Execute the forcerecook cook. -->
 			<Log Message="$(NewLine)========== IterativeCompare.xml: Non-Incremental Cook =========="/>
 			<Commandlet Name="Cook" Project="$(ProjectFile)" Arguments="$(CommonCookArgs) -forcerecook"/>
 
-			<!-- Write the files from forcerecook's zenstore out into the Saved/Cooked/Platform dir. -->
+			<!-- Copy all cook artifacts to the compare_incremental dir. Copy rather than move so they remain
+			     behind for the next instance of this job to use in an incremental cook. -->
+			<Log Message="$(NewLine)========== IterativeCompare.xml: Copy =========="/>
+			<Log Message="Delete Directories=$(CompareDirRecook)"/>
+			<Delete Directories="$(CompareDirRecook)"/>
+			<Log Message="Copy From=$(PlatformCookedDir)/*.* To=$(CompareDirRecook)"/>
+			<Copy From="$(PlatformCookedDir)/*.*" To="$(CompareDirRecook)" ErrorIfNotFound="true"/>
+
+			<!-- Write the files from forcerecook's zenstore out into CompareDirRecook. -->
 			<Log Message="$(NewLine)========== IterativeCompare.xml: ZenMirror ForceRecook Cook =========="/>
-			<Log Message="ZenMirror Project=$(ProjectFile) Platform=$(TargetPlatform)"/>
-			<ZenMirror Project="$(ProjectFile)" Platform="$(TargetPlatform)"/>
+			<Log Message="ZenMirror Project=$(ProjectFile) Platform=$(TargetPlatform) DestinationFileDir=$(CompareDirRecook)"/>
+			<ZenMirror Project="$(ProjectFile)" Platform="$(TargetPlatform)" DestinationFileDir="$(CompareDirRecook)"/>
 
 			<!-- Compare the incremental cook with the non-incremental cook. -->
 			<Log Message="$(NewLine)========== IterativeCompare.xml: DiffCook =========="/>
-			<Commandlet Name="DiffCook" Project="$(ProjectFile)" Arguments="-base=$(PlatformCookedDir) -comp=$(CompareDir) -targetplatform=$(TargetPlatform) -addedverbosity=error -removedverbosity=error -modifiedverbosity=warning -hide=modifiedpackages"/>
+			<Commandlet Name="DiffCook" Project="$(ProjectFile)" Arguments="-base=$(CompareDirRecook) -comp=$(CompareDirIncremental) -targetplatform=$(TargetPlatform) -addedverbosity=error -removedverbosity=error -modifiedverbosity=warning -hide=modifiedpackages"/>
 		</Node>
 	</Agent>
 
