Disable uploading symbols for Android and iOS

#android
#ios

diff --git a/Engine/Build/Graph/Tasks/BuildAndTestProject.xml b/Engine/Build/Graph/Tasks/BuildAndTestProject.xml
--- a/Engine/Build/Graph/Tasks/BuildAndTestProject.xml
+++ b/Engine/Build/Graph/Tasks/BuildAndTestProject.xml
@@ -462,7 +462,7 @@
 	<ForEach Name="TargetPlatform" Values="$(TargetPlatforms)" Separator="+">
 
 		<Property Name="IsDesktopPlatform" Value="ContainsItem('$(AllDesktopPlatforms)', '$(TargetPlatform)', '+')"/>
-		<Property Name="IsMobilePlatform" Value="'$(TargetPlatform)' == 'Android' or '$(TargetPlatform)' == 'IOS'"/>
+		<Property Name="IsMobilePlatform" Value="('$(TargetPlatform)' == 'Android' or '$(TargetPlatform)' == 'IOS')"/>
 		<!-- Local or virtualized platform -->
 		<Property Name="IsVirtualizedPlatform" Value="ContainsItem('$(AllVirtualizedPlatforms)', '$(TargetPlatform)', '+')"/>
 		<Property Name="IsLocalPlatform" Value="($(IsDesktopPlatform) Or $(IsVirtualizedPlatform))"/>
@@ -493,7 +493,7 @@
 		<Property Name="PlatformAgentName" Value="$(ProjectName) $(TargetPlatform)" />
 		<Property Name="PlatformCompileNodeName" Value="$(PreNodeName)Compile $(TargetPlatform)" />
 		<Property Name="UploadSymbolsNodeName" Value="" />
-		<Property Name="UploadSymbolsNodeName" Value="$(PreNodeName)Upload Symbols $(TargetPlatform)" If="'$(SymbolServerPath)' != ''" />
+		<Property Name="UploadSymbolsNodeName" Value="$(PreNodeName)Upload Symbols $(TargetPlatform)" If="'$(SymbolServerPath)' != '' And !$(IsMobilePlatform)" />
 		<Property Name="PlatformCookNodeName" Value="$(PreNodeName)Cook $(TargetPlatform)" />
 		<Property Name="PlatformStageNodeName" Value="$(PreNodeName)Stage $(TargetPlatform)" />
 		<Property Name="PlatformPackageNodeName" Value="$(PreNodeName)Package $(TargetPlatform)" />
