Fix issue with content only project + allow disabling test agent type

#rnx

diff --git a/Engine/Build/Graph/Tasks/BuildAndTestProject.xml b/Engine/Build/Graph/Tasks/BuildAndTestProject.xml
--- a/Engine/Build/Graph/Tasks/BuildAndTestProject.xml
+++ b/Engine/Build/Graph/Tasks/BuildAndTestProject.xml
@@ -54,6 +54,7 @@
 	<Option Name="TestAgentOverride" DefaultValue="" Description="If set, other logic is ignored and this agent type is used for all test work." />
 	<Option Name="TargetPlatformTestAgentOverride" DefaultValue="DevkitAutomation" Description="Override agent type use to run target platform tests"/>
 	<Option Name="ClientCompileAgentOverride" DefaultValue="" Description="Specifies the agent type to use for C++ compilation jobs for target platform binaries." />
+	<Option Name="DontUseTestAgent" Restrict="true|false" DefaultValue="false" Description="Avoid using test agent pool to run tests"/>
 	<Option Name="UnrealInsightsBuildConfiguration" DefaultValue="Shipping" Restrict="Development|Debug|Shipping" Description="In which configuration to build Unreal Insights, if included. Shipping, Debug, or Development." />
 	<Option Name="IncludeUnrealInsightsInToolCompile" DefaultValue="false" Description="Include Unreal Insights into Compile Tools prerequisites" />
 	<Option Name="IncludeEpicWebHelperInToolCompile" DefaultValue="false" Description="Include EpicWebHelper into Compile Tools prerequisites" />
@@ -235,6 +236,7 @@
 		<!-- Use incremental agents for building? -->
 		<Property Name="HostAgentType" Value="Incremental$(HostAgentType)" If="$(UseIncrementalAgents)" />
 		<Property Name="TestAgentType" Value="Test$(EditorPlatform)" />
+		<Property Name="TestAgentType" Value="$(EditorPlatform)" If="$(DontUseTestAgent)" />
 		<!-- Use overrides for agent types? -->
 		<Property Name="HostAgentType" Value="$(AgentOverride)" If="'$(AgentOverride)' != ''" />
 		<Property Name="TestAgentType" Value="$(TestAgentOverride)"  If="'$(TestAgentOverride)' != ''" />
@@ -273,7 +275,7 @@
 				<Switch>
 					<Case If="$(IsContentOnlyProject)">
 						<Command Name="BuildCookRun" Arguments="-project=&quot;$(RootDir)/$(TargetProject)&quot; -platform=$(EditorPlatform) -configuration=Development -build -notools -ubtargs=&quot;$(GenericCompileArguments) $(ExtraEditorCompileArguments)&quot;"/>
-						<Tag Files="$(RootDir)/$(ProjectPath)/Binaries/..." Filter="$(EditorTargetName).target" With="#$(CompileEditor)"/>
+						<Tag Files="Engine/Binaries/...;$(ProjectPath)/Binaries/..." Filter="UnrealEditor.target;$(EditorTargetName).target" With="#$(CompileEditor)"/>
 						<TagReceipt Files="#$(CompileEditor)" BuildProducts="true" With="#$(CompileEditor)"/>
 					</Case>
 					<Default>
@@ -360,6 +362,7 @@
 	<ForEach Name="EditorPlatform" Values="$(CookedEditorPlatforms)" Separator="+">
 		<Property Name="HostAgentType" Value="$(EditorPlatform)" />
 		<Property Name="TestAgentType" Value="Test$(EditorPlatform)" />
+		<Property Name="TestAgentType" Value="$(EditorPlatform)" If="$(DontUseTestAgent)" />
 		<Property Name="TestAgentType" Value="$(TestAgentOverride)" If="'$(TestAgentOverride)' != ''"/>
 
 		<Property Name="EditorPlatformPath" Value="$(EditorPlatform)" />
@@ -474,6 +477,7 @@
 		<!-- On consoles we can use agents with devkit setup, but locally hosted tests need a Test agent-->
 		<Property Name="TestAgentType" Value="$(TargetPlatformTestAgentOverride)" />
 		<Property Name="TestAgentType" Value="Test$(HostAgentType)" If="$(IsDesktopPlatform)"/>
+		<Property Name="TestAgentType" Value="$(HostAgentType)" If="$(IsDesktopPlatform) and $(DontUseTestAgent)" />
 
 		<!-- Use incremental agents for building? -->
 		<Property Name="CompileAgentType" Value="$(HostAgentType)" />
@@ -547,9 +551,9 @@
 			<Node Name="$(PlatformCompileNodeName)" Requires="$(PlatformToolsNodeName)">
 				<Switch>
 					<Case If="$(IsContentOnlyProject)">
-						<Command Name="BuildGame" Arguments="$(BCRArgs) -configuration=$(TargetConfigurations) -notools -ubtargs=&quot;-BuildVersion=$(BuildVersion) $(GenericCompileArguments) $(ExtraTargetCompileArguments)&quot;" />
-						<Tag Files="$(RootDir)/$(ProjectPath)/Binaries/..." Filter="$(TargetName).target" With="#$(PlatformCompileNodeName)"/>
-						<TagReceipt Files="#$(PlatformCompileNodeName)" BuildProducts="true" With="#$(PlatformCompileNodeName)"/>
+						<ForEach Name="TargetConfiguration" Values="$(TargetConfigurations)" Separator="+">
+							<Compile Target="UnrealGame" Platform="$(TargetPlatform)" Configuration="$(TargetConfiguration)" Arguments="-BuildVersion=&quot;$(BuildVersion)&quot; $(GenericCompileArguments) $(ExtraTargetCompileArguments)"/>
+						</ForEach>
 					</Case>
 					<Default>
 						<ForEach Name="TargetConfiguration" Values="$(TargetConfigurations)" Separator="+">
