Fix content only project not getting staged - second try

#rnx

diff --git a/Engine/Build/Graph/Tasks/BuildAndTestProject.xml b/Engine/Build/Graph/Tasks/BuildAndTestProject.xml
--- a/Engine/Build/Graph/Tasks/BuildAndTestProject.xml
+++ b/Engine/Build/Graph/Tasks/BuildAndTestProject.xml
@@ -551,9 +551,11 @@
 			<Node Name="$(PlatformCompileNodeName)" Requires="$(PlatformToolsNodeName)">
 				<Switch>
 					<Case If="$(IsContentOnlyProject)">
-						<ForEach Name="TargetConfiguration" Values="$(TargetConfigurations)" Separator="+">
-							<Compile Target="UnrealGame" Platform="$(TargetPlatform)" Configuration="$(TargetConfiguration)" Arguments="-BuildVersion=&quot;$(BuildVersion)&quot; $(GenericCompileArguments) $(ExtraTargetCompileArguments)"/>
-						</ForEach>
+						<Command Name="BuildGame" Arguments="$(BCRArgs) -configuration=$(TargetConfigurations) -notools -ubtargs=&quot;-BuildVersion=$(BuildVersion) $(GenericCompileArguments) $(ExtraTargetCompileArguments)&quot;" />
+						<Tag Files="$(ProjectPath)/Binaries/..." Filter="$(TargetName)*.target" With="#$(PlatformCompileNodeName)"/>
+						<Tag Files="$(ProjectPath)/Binaries/..." Filter="UnrealGame*.target" With="#UnrealGameTarget"/>
+						<Copy Files="#UnrealGameTarget" From="$(ProjectPath)" To="Engine" Tag="#$(PlatformCompileNodeName)"/>
+						<TagReceipt Files="#$(PlatformCompileNodeName)" BuildProducts="true" With="#$(PlatformCompileNodeName)"/>
 					</Case>
 					<Default>
 						<ForEach Name="TargetConfiguration" Values="$(TargetConfigurations)" Separator="+">
