Allow for experimental types included in the iterative validate job.

diff --git a/Engine/Build/Graph/Tests/IterativeValidate.xml b/Engine/Build/Graph/Tests/IterativeValidate.xml
--- a/Engine/Build/Graph/Tests/IterativeValidate.xml
+++ b/Engine/Build/Graph/Tests/IterativeValidate.xml
@@ -27,9 +27,13 @@
 	<Property Name="IterativeArgs" Value="-zenstore -cookincremental" If="'$(IterativeMode)' == 'incremental'"/>
 	<Property Name="IterativeArgs" Value="-iterate -ini:Editor:[CookSettings]:HybridIterativeEnabled=false -ini:Editor:[CookSettings]:IterativeUseClassFilters=true" If="'$(IterativeMode)' == 'iterative'"/>
 
+	<!-- Config Options -->
+	<Option Name="CustomConfig" DefaultValue="" Description="What custom configuration to apply (see -CustomConfig commandline argument)"/>
+
 	<!-- Agent=IterativeValidate variables -->
 	<Property Name="EditorPlatform" Value="Win64"/>
 	<Property Name="CommonCookArgs" Value="-targetplatform=$(TargetPlatform) $(IterativeArgs)"/>
+	<Property Name="CommonCookArgs" Value="$(CommonCookArgs) -CustomConfig=$(CustomConfig)" If="'$(CustomConfig)' != ''"/>
 
 	<Property Name="CompileNodeName" Value="Compile $(EditorPlatform) $(ProjectName)"/>
 	<Property Name="CookIterativeValidateSinglePhaseName" Value="Cook IterativeValidateSinglePhase $(ProjectName) $(TargetPlatform)"/>
