Publish Zen Android snapshot descriptors to the appropriate location in the build index.

[RN] fix

diff --git a/Engine/Build/Graph/Tasks/BuildAndTestProject.xml b/Engine/Build/Graph/Tasks/BuildAndTestProject.xml
--- a/Engine/Build/Graph/Tasks/BuildAndTestProject.xml
+++ b/Engine/Build/Graph/Tasks/BuildAndTestProject.xml
@@ -654,7 +654,10 @@
 					<Property Name="SnapshotIdentifier" Value="$(UE_HORDE_TEMPLATEID).$(BuildVersion).$(UE_HORDE_JOBID:-6).$(UE_HORDE_STEPID)" If="'$(UE_HORDE_JOBID)' != '' and '$(UE_HORDE_STEPID)' != '' and '$(UE_HORDE_TEMPLATEID)' != ''"/>
 					<Property Name="SnapshotFilenamePrefix" Value="SnapshotDescriptor"/>
 					<Property Name="SnapshotFilenamePrefix" Value="$(UE_HORDE_TEMPLATEID)" If="'$(UE_HORDE_TEMPLATEID)' != ''"/>
-					<Property Name="SnapshotLongTermStorageDir" Value="$(NetworkSnapshotDescFileRoot)/$(EscapedBranch)/$(ProjectName)/$(StagedPlatformFolder)/$(Change)"  If="'$(NetworkSnapshotDescFileRoot)' != ''"/>
+					<Property Name="SnapshotLongTermPlatformFolder" Value="$(StagedPlatformFolder)" />
+					<Property Name="SnapshotLongTermPlatformFolder" Value="Android" If="'$(TargetPlatform)' == 'Android'" />
+					<Property Name="SnapshotLongTermStorageDir" Value="$(NetworkSnapshotDescFileRoot)/$(EscapedBranch)/$(ProjectName)/$(SnapshotLongTermPlatformFolder)/$(Change)"  If="'$(NetworkSnapshotDescFileRoot)' != ''"/>
+
 					<Do If="$(ShouldSnapshotPublishCloud)">
 						<Property Name="SnapshotBaseDescriptorFile" Value=""/>
 						<Property Name="SnapshotBaseDescriptorFile" Value="$(NetworkSnapshotDescFileRoot)/$(EscapedBranch)/$(ProjectName)/$(StagedPlatformFolder)/$(UE_HORDE_LAST_WARNING_CL)/$(SnapshotFilenamePrefix)-cloud.json"  If="'$(NetworkSnapshotDescFileRoot)' != '' and '$(UE_HORDE_LAST_WARNING_CL)' != ''"/>
