Remove UBAPrintSummary
#rnx

diff --git a/Engine/Build/Graph/Tests/DDCVerify.xml b/Engine/Build/Graph/Tests/DDCVerify.xml
--- a/Engine/Build/Graph/Tests/DDCVerify.xml
+++ b/Engine/Build/Graph/Tests/DDCVerify.xml
@@ -20,7 +20,7 @@
 
 	<Option Name="WithUBA" Restrict="true|false" DefaultValue="true" Description="Whether to use UBA compile or not"/>
 	<Property Name="UBACompileArgs" Value=""/>
-	<Property Name="UBACompileArgs" Value="-UBA -ActionStallReportTime=300 -UBALinkRemote -UBAPrintSummary" If="$(WithUBA)"/>
+	<Property Name="UBACompileArgs" Value="-UBA -ActionStallReportTime=300 -UBALinkRemote" If="$(WithUBA)"/>
 	
 	<!-- DDC Verify Options -->
 	<Option Name="DDCVerifyDefaultRate" Restrict="\d*" DefaultValue="1" Description="The default verification rate in percent"/>
diff --git a/Engine/Build/Graph/Tests/IterativeCompare.xml b/Engine/Build/Graph/Tests/IterativeCompare.xml
--- a/Engine/Build/Graph/Tests/IterativeCompare.xml
+++ b/Engine/Build/Graph/Tests/IterativeCompare.xml
@@ -21,7 +21,7 @@
 
 	<Option Name="WithUBA" Restrict="true|false" DefaultValue="true" Description="Whether to use UBA compile or not"/>
 	<Property Name="UBACompileArgs" Value=""/>
-	<Property Name="UBACompileArgs" Value="-UBA -ActionStallReportTime=300 -UBALinkRemote -UBAPrintSummary" If="$(WithUBA)"/>
+	<Property Name="UBACompileArgs" Value="-UBA -ActionStallReportTime=300 -UBALinkRemote" If="$(WithUBA)"/>
 	
 	<!-- Agent=IterativeCompareAgent variables -->
 	<Property Name="FullPathProjectDir" Value="$(RootDir)/$(ProjectDir)"/>
diff --git a/Engine/Build/Graph/Tests/IterativeValidate.xml b/Engine/Build/Graph/Tests/IterativeValidate.xml
--- a/Engine/Build/Graph/Tests/IterativeValidate.xml
+++ b/Engine/Build/Graph/Tests/IterativeValidate.xml
@@ -19,7 +19,7 @@
 
 	<Option Name="WithUBA" Restrict="true|false" DefaultValue="true" Description="Whether to use UBA compile or not"/>
 	<Property Name="UBACompileArgs" Value=""/>
-	<Property Name="UBACompileArgs" Value="-UBA -ActionStallReportTime=300 -UBALinkRemote -UBAPrintSummary" If="$(WithUBA)"/>
+	<Property Name="UBACompileArgs" Value="-UBA -ActionStallReportTime=300 -UBALinkRemote" If="$(WithUBA)"/>
 	
 	<!-- Iterative Options -->
 	<Option Name="IterativeMode" DefaultValue="incremental" Restrict="iterative|incremental" Description="Whether to iteratively cook using Incremental or using the legacy iterative"/>
