Package Project - C++ - Development - Windows - Linux

#rnx

diff --git a/Engine/Build/Graph/Tasks/BuildAndTestProject.xml b/Engine/Build/Graph/Tasks/BuildAndTestProject.xml
--- a/Engine/Build/Graph/Tasks/BuildAndTestProject.xml
+++ b/Engine/Build/Graph/Tasks/BuildAndTestProject.xml
@@ -496,10 +496,19 @@
 		<Property Name="IsVirtualizedPlatform" Value="ContainsItem('$(AllVirtualizedPlatforms)', '$(TargetPlatform)', '+')"/>
 		<Property Name="IsLocalPlatform" Value="($(IsDesktopPlatform) Or $(IsVirtualizedPlatform))"/>
 
+		<Property Name="IsTargetPlatformMac" Value="('$(TargetPlatform)' == 'Mac' or '$(TargetPlatform)' == 'IOS' or '$(TargetPlatform)' == 'tvOS')"/>
+		<Property Name="IsTargetPlatformLinux" Value="('$(TargetPlatform)' == 'Linux')"/>
+		<Property Name="IsTargetPlatformWindows" Value="(!$(IsTargetPlatformMac) and !$(IsTargetPlatformLinux))"/>
+		
 		<!-- Declare type of host required to build & cook this platform-->
-		<Property Name="HostAgentType" Value="Win64" />
-		<Property Name="HostAgentType" Value="Mac" If="'$(TargetPlatform)' == 'Mac' or '$(TargetPlatform)' == 'IOS' or '$(TargetPlatform)' == 'tvOS'"/>
-		<Property Name="HostAgentType" Value="Linux" If="'$(TargetPlatform)' == 'Linux'"/>
+		<Property Name="HostAgentType" Value="$(EditorPlatformWin64AgentOverride)" /> 
+		<Property Name="HostAgentType" Value="$(EditorPlatformMacAgentOverride)" If="$(IsTargetPlatformMac)"/>
+		<Property Name="HostAgentType" Value="$(EditorPlatformLinuxAgentOverride)" If="$(IsTargetPlatformLinux)"/>
+
+		<Property Name="IsEditorAgentOverride" Value="false"/>
+		<Property Name="IsEditorAgentOverride" Value="true" If="$(IsTargetPlatformWindows) and '$(EditorPlatformWin64AgentOverride)'!='Win64'"/>
+		<Property Name="IsEditorAgentOverride" Value="true" If="$(IsTargetPlatformMac) and '$(EditorPlatformMacAgentOverride)'!='Mac'"/>
+		<Property Name="IsEditorAgentOverride" Value="true" If="$(IsTargetPlatformLinux) and '$(EditorPlatformLinuxAgentOverride)'!='Linux'"/>
 		<!-- On consoles we can use agents with devkit setup, but locally hosted tests need a Test agent-->
 		<Property Name="TestAgentType" Value="$(TargetPlatformTestAgentOverride)" />
 		<Property Name="TestAgentType" Value="Test$(HostAgentType)" If="$(IsDesktopPlatform)"/>
@@ -516,22 +525,24 @@
 		<Property Name="CompileAgentType" Value="$(ClientCompileAgentOverride)" If="'$(ClientCompileAgentOverride)' != ''"/>
 		<Property Name="TestAgentType" Value="$(TestAgentOverride)" If="'$(TestAgentOverride)' != ''"/>
 
+		<Property Name="PreNodeNameWithPlatform" Value="$(PreNodeName)"/>
+		<Property Name="PreNodeNameWithPlatform" Value="[$(HostAgentType)]$(PreNodeName)" If="$(IsEditorAgentOverride)"/>
 		<!-- Building and cooking will require these nodes from the appropriate host. They should have been built above... -->
 		<Property Name="PlatformToolsNodeName" Value="$(PreNodeName)Compile Tools $(HostAgentType)" />
 		<Property Name="PlatformEditorNodeName" Value="$(PreNodeName)Compile Editor $(HostAgentType)" />
 		<!-- Define these names up front as we use some of them several times -->
 		<Property Name="PlatformAgentName" Value="$(ProjectName) $(TargetPlatform)" />
-		<Property Name="PlatformCompileNodeName" Value="$(PreNodeName)Compile $(TargetPlatform)" />
+		<Property Name="PlatformCompileNodeName" Value="$(PreNodeNameWithPlatform)Compile $(TargetPlatform)" />
 		<Property Name="UploadSymbolsNodeName" Value="" />
-		<Property Name="UploadSymbolsNodeName" Value="$(PreNodeName)Upload Symbols $(TargetPlatform)" If="'$(SymbolServerPath)' != '' And !$(IsMobilePlatform)" />
-		<Property Name="PlatformCookNodeName" Value="$(PreNodeName)Cook $(TargetPlatform)" />
-		<Property Name="PlatformStageNodeName" Value="$(PreNodeName)Stage $(TargetPlatform)" />
-		<Property Name="PlatformPackageNodeName" Value="$(PreNodeName)Package $(TargetPlatform)" />
+		<Property Name="UploadSymbolsNodeName" Value="$(PreNodeNameWithPlatform)Upload Symbols $(TargetPlatform)" If="'$(SymbolServerPath)' != '' And !$(IsMobilePlatform)" />
+		<Property Name="PlatformCookNodeName" Value="$(PreNodeNameWithPlatform)Cook $(TargetPlatform)" />
+		<Property Name="PlatformStageNodeName" Value="$(PreNodeNameWithPlatform)Stage $(TargetPlatform)" />
+		<Property Name="PlatformPackageNodeName" Value="$(PreNodeNameWithPlatform)Package $(TargetPlatform)" />
 		<Property Name="PlatformPackageArtifactName" Value="$(PlatformPackageNodeName)" />
-		<Property Name="PlatformHashNodeName" Value="$(PreNodeName)Hash $(TargetPlatform)" />
-		<Property Name="PlatformPublishStagedNodeName" Value="$(PreNodeName)Publish Staged $(TargetPlatform)" />
-		<Property Name="PlatformPublishPackagedNodeName" Value="$(PreNodeName)Publish Packaged $(TargetPlatform)" />
-		<Property Name="PlatformCheckBuildSizeNodeName" Value="$(PreNodeName)Check Build Size $(TargetPlatform)" />
+		<Property Name="PlatformHashNodeName" Value="$(PreNodeNameWithPlatform)Hash $(TargetPlatform)" />
+		<Property Name="PlatformPublishStagedNodeName" Value="$(PreNodeNameWithPlatform)Publish Staged $(TargetPlatform)" />
+		<Property Name="PlatformPublishPackagedNodeName" Value="$(PreNodeNameWithPlatform)Publish Packaged $(TargetPlatform)" />
+		<Property Name="PlatformCheckBuildSizeNodeName" Value="$(PreNodeNameWithPlatform)Check Build Size $(TargetPlatform)" />
 		<Property Name="PlatformStagedArtifactNodeName" Value="$(PlatformPublishStagedNodeName) (Horde)" />
 		<Property Name="HordeArtifactPlatformPackagedNodeName" Value="$(PlatformPublishPackagedNodeName) (Horde)" />
 		<Property Name="RequirePlatformStaged" Value="$(PlatformStageNodeName)" />
diff --git a/Engine/Build/Graph/Tasks/Inc/PlatformOptions.xml b/Engine/Build/Graph/Tasks/Inc/PlatformOptions.xml
--- a/Engine/Build/Graph/Tasks/Inc/PlatformOptions.xml
+++ b/Engine/Build/Graph/Tasks/Inc/PlatformOptions.xml
@@ -34,6 +34,10 @@
 <Option Name="TargetPlatforms" DefaultValue="" Description="Alternative to 'WithFoo', a list of Platforms to build and test. Android+IOS etc"/>
 <Option Name="CookedEditorPlatforms" DefaultValue="" Description="Alternative to 'WithFooCookedEditor', a list of editors to build and cooked. Win64+Mac etc"/>
 
+<Option Name="EditorPlatformWin64AgentOverride" DefaultValue="Win64" Restrict="Win64|Linux|Mac" Description="Platform to use to build an editor for Win64 target"/>
+<Option Name="EditorPlatformLinuxAgentOverride" DefaultValue="Linux" Restrict="Win64|Linux|Mac" Description="Platform to use to build an editor for Linux target"/>
+<Option Name="EditorPlatformMacAgentOverride" DefaultValue="Mac" Restrict="Win64|Linux|Mac" Description="Platform to use to build an editor for Mac target"/>
+
 <!-- Support simple One+Two+Three style -->
 <Option Name="TargetConfigurations" DefaultValue="" Description="Configuration(s) to build"/>
 <Option Name="PackageConfigurations" DefaultValue="$(TargetConfigurations)" Description="Configuration to package"/>
@@ -136,9 +140,9 @@
 		<Error Message="Platform $(TargetPlatform) is not a recognized platform!" />
 	</Do>
 	<!-- Declare type of host required to build & cook this platform-->
-	<Property Name="HostAgentType" Value="Win64" />
-	<Property Name="HostAgentType" Value="Mac" If="'$(TargetPlatform)' == 'Mac' or '$(TargetPlatform)' == 'IOS' or '$(TargetPlatform)' == 'TVOS'"/>
-	<Property Name="HostAgentType" Value="Linux" If="'$(TargetPlatform)' == 'Linux'"/>
+	<Property Name="HostAgentType" Value="$(EditorPlatformWin64AgentOverride)" />
+	<Property Name="HostAgentType" Value="$(EditorPlatformMacAgentOverride)" If="'$(TargetPlatform)' == 'Mac' or '$(TargetPlatform)' == 'IOS' or '$(TargetPlatform)' == 'TVOS'"/>
+	<Property Name="HostAgentType" Value="$(EditorPlatformLinuxAgentOverride)" If="'$(TargetPlatform)' == 'Linux'"/>
 
 	<Do If="!ContainsItem('$(RequiredEditorPlatforms)','$(HostAgentType)', '+')">
 		<Property Name="RequiredEditorPlatforms" Value="$(RequiredEditorPlatforms)+$(HostAgentType)" />
