Use Horde artifacts to get game binaries

#rnx

diff --git a/Engine/Build/Graph/Tasks/BuildAndTestProject.xml b/Engine/Build/Graph/Tasks/BuildAndTestProject.xml
--- a/Engine/Build/Graph/Tasks/BuildAndTestProject.xml
+++ b/Engine/Build/Graph/Tasks/BuildAndTestProject.xml
@@ -545,11 +545,12 @@
 		<Property Name="PlatformCheckBuildSizeNodeName" Value="$(PreNodeNameWithPlatform)Check Build Size $(TargetPlatform)" />
 		<Property Name="PlatformStagedArtifactNodeName" Value="$(PlatformPublishStagedNodeName) (Horde)" />
 		<Property Name="HordeArtifactPlatformPackagedNodeName" Value="$(PlatformPublishPackagedNodeName) (Horde)" />
+		<Property Name="HordeArtifactPlatformRetrieve" Value="$(PreNodeNameWithPlatform)Retrieve Build Artifact $(TargetPlatform)" />
 		<Property Name="RequirePlatformStaged" Value="$(PlatformStageNodeName)" />
 		<Property Name="RequirePlatformStaged" Value="$(RequirePlatformStaged);$(PlatformStagedArtifactNodeName)" If="$(Horde)"/>
 		<Property Name="RequirePlatformStaged" Value="$(RequirePlatformStaged);$(PlatformPublishStagedNodeName)" If="!$(SkipPublish) And $(PublishBuilds)"/>
 		<Property Name="RequirePlatformPackaged" Value="$(PlatformPackageNodeName)" />
-		<Property Name="RequirePlatformPackaged" Value="$(RequirePlatformPackaged);#$(HordeArtifactPlatformPackagedNodeName) Complete" If="$(Horde)"/>
+		<Property Name="RequirePlatformPackaged" Value="$(RequirePlatformPackaged);$(HordeArtifactPlatformPackagedNodeName)" If="$(Horde)"/>
 		<Property Name="RequirePlatformPackaged" Value="$(RequirePlatformPackaged);$(PlatformPublishPackagedNodeName)"  If="!$(SkipPublish) And $(PublishBuilds)"/>
 		
 		<!-- Set CookPlatform since it isn't always the same as the target -->
@@ -575,8 +576,7 @@
 		<!-- Support additional stage and package arguments that are only active on some platforms -->
 		<Property Name="StageAndPackageArguments" Value="$(ExtraStageAndPackageArguments)" />
 		<Property Name="StageAndPackageArguments" Value="$(StageAndPackageArguments) $(TargetPlatformFilteredExtraStageAndPackageArguments)" If="'$(TargetPlatformFilteredExtraStageAndPackageArgumentsFilter)' != '' and ContainsItem('$(TargetPlatformFilteredExtraStageAndPackageArgumentsFilter)', '$(TargetPlatform)', '+')" />
-		<Property Name="PackageOnlyArguments" Value="$(ExtraPackageOnlyArguments)"/>
-		<Property Name="PackageOnlyArguments" Value="$(PackageOnlyArguments) -archive -archivedirectory=&quot;$(ProjectOutputDirectory)/$(StagedPlatformFolder)&quot;" If="$(IsMobilePlatform)" />
+		<Property Name="PackageOnlyArguments" Value="$(ExtraPackageOnlyArguments)" />
 
 		<Property Name="BuildSizeLimit" Value=""/>
 		<Regex Pattern="(?&lt;!\w)$(TargetPlatform):([0-9]+)" Capture="BuildSizeLimit" Input="$(CheckBuildSizePlatforms)"/>
@@ -659,7 +659,7 @@
 			<Property Name="ShouldSnapshotPublish" Value="true" If="$(IsBuildMachine) and !$(IsPreflight) and $(ShouldSnapshotPublishCloud) or $(ShouldSnapshotPublishZen)"/>
 			<Do If="$(PublishBuilds) and $(ShouldSnapshotPublish)">
 				<Property Name="PlatformSnapshotPublishNodeName" Value="$(PreNodeName)Snapshot Publish $(TargetPlatform)" />
-				<Node Name="$(PlatformSnapshotPublishNodeName)" Requires="$(PlatformHashNodeName)" Produces="#SnapshotDescriptor_$(StagedPlatformFolder)">
+				<Node Name="$(PlatformSnapshotPublishNodeName)" Requires="$(PlatformHashNodeName)">
 					<Property Name="SnapshotLocalDir" Value="$(RootDir)/Snapshots" />
 					<Property Name="SnapshotIdentifier" Value="$(EscapedBranch)/$(BuildNamePath)"/>
 					<Property Name="SnapshotIdentifier" Value="$(UE_HORDE_TEMPLATEID).$(BuildVersion).$(UE_HORDE_JOBID:-6).$(UE_HORDE_STEPID)" If="'$(UE_HORDE_JOBID)' != '' and '$(UE_HORDE_STEPID)' != '' and '$(UE_HORDE_TEMPLATEID)' != ''"/>
@@ -686,27 +686,32 @@
 						</Do>
 					</Do>
 					<Tag Files="$(SnapshotLocalDir)/..." With="#SnapshotDescriptor_$(StagedPlatformFolder)"/>
+
+					<CreateArtifact Name="ci-snapshot-$(StagedPlatformFolder)" Description="Snapshot ($(StagedPlatformFolder))" Type="ugs-snapshot" BaseDir="Snapshots" Keys="ugs-project=$(TargetProject)" Metadata="ArchiveType=SnapshotDescriptor_$(StagedPlatformFolder);dashboard-category=Snapshots" Files="#SnapshotDescriptor_$(StagedPlatformFolder)"/>
 				</Node>
-				<Artifact Name="ci-snapshot-$(StagedPlatformFolder)" Description="Snapshot ($(StagedPlatformFolder))" Type="ugs-snapshot" BasePath="Snapshots" Keys="ugs-project=$(TargetProject)" Metadata="ArchiveType=SnapshotDescriptor_$(StagedPlatformFolder);dashboard-category=Snapshots" Tag="#SnapshotDescriptor_$(StagedPlatformFolder)"/>
 				<Property Name="PlatformPublishStageDependencyName" Value="$(PlatformSnapshotPublishNodeName)" />
 			</Do>
 			
 			<!-- Publish (Staged) -->
 			<Node Name="$(PlatformPublishStagedNodeName)" Requires="$(PlatformPublishStageDependencyName)">
+				<Error Message="NetworkOutputDirectory is not set" If="'$(NetworkOutputDirectory)' == ''"/>
 				<Copy Files="..." From="$(ProjectOutputDirectory)/$(StagedPlatformFolder)" To="$(NetworkOutputDirectory)/$(StagedPlatformFolder)/Staged"/>
 				<Copy Files="..." From="$(ShaderSymbolsDirectory)/*$(TargetPlatform)*" To="$(NetworkOutputDirectory)/$(StagedPlatformFolder)/ShaderSymbols"/>
 			</Node>
 
 			<!-- Publish (Staged) (Horde) -->
-			<Node Name="$(PlatformStagedArtifactNodeName)" Requires="$(PlatformStageNodeName)" Produces="#$(PlatformStageNodeName) Files">
-				<Tag Files="$(ProjectOutputDirectory)/$(StagedPlatformFolder)/..." With="#$(PlatformStageNodeName) Files"/>
-			</Node>
-			<Artifact Name="$(PlatformPackageArtifactName) (Staged)" Type="staged-build" BasePath="$(ProjectOutputDirectory)/$(StagedPlatformFolder)" Keys="Project=$(ProjectName)" Metadata="Stream=$(Branch);Platform=$(StagedPlatformFolder);BuildName=$(BuildNamePath);dashboard-category=Staged Builds" Tag="#$(PlatformStageNodeName) Files"/>
+			<Do If="$(Horde)">
+				<Node Name="$(PlatformStagedArtifactNodeName)" Requires="$(PlatformStageNodeName)">
+					<Tag Files="$(ProjectOutputDirectory)/$(StagedPlatformFolder)/..." With="#$(PlatformStageNodeName) Files"/>
+
+					<CreateArtifact Name="Staged-$(StagedPlatformFolder)" Type="staged-build" BaseDir="$(ProjectOutputDirectory)/$(StagedPlatformFolder)" Keys="Project=$(ProjectName)" Metadata="Stream=$(Branch);Platform=$(StagedPlatformFolder);BuildName=$(BuildNamePath);dashboard-category=Staged Builds" Files="#$(PlatformStageNodeName) Files"/>
+				</Node>
+			</Do>
 
 			<!-- Package -->
 			<Node Name="$(PlatformPackageNodeName)" Requires="$(PlatformPublishStageDependencyName)">
 				<ForEach Name="PackageConfiguration" Values="$(PackageConfigurations)" Separator="+">
-					<Command Name="BuildCookRun" Arguments="$(BCRArgs) -configuration=$(PackageConfiguration) $(StageAndPackageArguments) -skipbuild -skipcook -skipstage -skippak -package $(PackageOnlyArguments)" />
+					<Command Name="BuildCookRun" Arguments="$(BCRArgs) -configuration=$(PackageConfiguration) $(StageAndPackageArguments) -skipbuild -skipcook -skipstage -skippak -package $(PackageOnlyArguments) -archive -archivedirectory=&quot;$(ProjectOutputDirectory)/Packages/$(StagedPlatformFolder)/$(PackageConfiguration)&quot;" />
 				</ForEach>
 			</Node>
 			
@@ -719,27 +724,24 @@
 			</Node>
 			
 			<!-- Publish (Packages) (Horde) -->
-			<Property Name="PackagedArtifactTags" Value=""/>
-			<Property Name="HordePackageSkipBCRArguments" Value="-stage -package"/>
-			<Property Name="HordePackageSkipBCRArguments" Value="-skipstage -skippackage" If="'$(TargetPlatform)' == 'Mac'"/>
+			<Do  If="$(Horde)">
+				<Property Name="HordePackageSkipBCRArguments" Value="-stage -package"/>
+				<Property Name="HordePackageSkipBCRArguments" Value="-skipstage -skippackage" If="'$(TargetPlatform)' == 'Mac'"/>
 
-			<ForEach Name="PackageConfiguration" Values="$(PackageConfigurations)" Separator="+">
-				<Property Name="PackagedArtifactTags" Value="#$(PlatformPackageNodeName) $(PackageConfiguration) Files;$(PackagedArtifactTags)"/>
-			</ForEach>
-			<Node Name="$(HordeArtifactPlatformPackagedNodeName)" Requires="$(PlatformPackageNodeName)" Produces="$(PackagedArtifactTags);#$(HordeArtifactPlatformPackagedNodeName) Complete">
-				<ForEach Name="PackageConfiguration" Values="$(PackageConfigurations)" Separator="+">
-					<Command Name="BuildCookRun" Arguments="$(BCRArgs) -configuration=$(PackageConfiguration) $(StageAndPackageArguments) -skipbuild -skipcook -skippak $(HordePackageSkipBCRArguments) -archive -archivedirectory=&quot;$(ProjectOutputDirectory)/Packages/$(StagedPlatformFolder)/$(PackageConfiguration)&quot;" />
-					<Tag Files="$(ProjectOutputDirectory)/Packages/$(StagedPlatformFolder)/$(PackageConfiguration)/..." With="#$(PlatformPackageNodeName) $(PackageConfiguration) Files"/>
-				</ForEach>
-			</Node>
-			<ForEach Name="PackageConfiguration" Values="$(PackageConfigurations)" Separator="+">
-				<Artifact Name="$(PlatformPackageArtifactName) ($(PackageConfiguration))" Type="packaged-build" BasePath="$(ProjectOutputDirectory)/Packages/$(StagedPlatformFolder)/$(PackageConfiguration)" Keys="Project=$(ProjectName)" Metadata="Stream=$(Branch);Platform=$(StagedPlatformFolder);Configuration=$(PackageConfiguration);BuildName=$(BuildNamePath);dashboard-category=$(PackageConfiguration) Builds" Tag="#$(PlatformPackageNodeName) $(PackageConfiguration) Files"/>
-			</ForEach>			
+				<Node Name="$(HordeArtifactPlatformPackagedNodeName)" Requires="$(PlatformPackageNodeName)">
+					<ForEach Name="PackageConfiguration" Values="$(PackageConfigurations)" Separator="+">
+						<Command Name="BuildCookRun" Arguments="$(BCRArgs) -configuration=$(PackageConfiguration) $(StageAndPackageArguments) -skipbuild -skipcook -skippak $(HordePackageSkipBCRArguments) -archive -archivedirectory=&quot;$(ProjectOutputDirectory)/Packages/$(StagedPlatformFolder)/$(PackageConfiguration)&quot;" />
+						<Tag Files="$(ProjectOutputDirectory)/Packages/$(StagedPlatformFolder)/$(PackageConfiguration)/..." With="#$(PlatformPackageNodeName) $(PackageConfiguration) Files"/>
+
+						<CreateArtifact Name="Packaged-$(StagedPlatformFolder)-$(PackageConfiguration)" Type="packaged-build" BaseDir="$(ProjectOutputDirectory)/Packages/$(StagedPlatformFolder)/$(PackageConfiguration)" Keys="Project=$(ProjectName)" Metadata="Stream=$(Branch);Platform=$(StagedPlatformFolder);Configuration=$(PackageConfiguration);BuildName=$(BuildNamePath);dashboard-category=$(PackageConfiguration) Builds" Files="#$(PlatformPackageNodeName) $(PackageConfiguration) Files"/>
+					</ForEach>
+				</Node>
+			</Do>
 
 			<!-- Check Build Size -->
-			<Node Name="$(PlatformCheckBuildSizeNodeName)" Requires="$(PlatformPublishPackagedNodeName)">
+			<Node Name="$(PlatformCheckBuildSizeNodeName)" Requires="$(PlatformPackageNodeName)">
 				<ForEach Name="PackageConfiguration" Values="$(PackageConfigurations)" Separator="+">
-					<Command Name="CheckBuildSize" Arguments="-project=$(TargetName) -platform=$(TargetPlatform) -packagesize=$(BuildSizeLimit) -build=&quot;$(NetworkOutputDirectory)&quot; -configuration=$(PackageConfiguration)" />
+					<Command Name="CheckBuildSize" Arguments="-project=$(TargetName) -platform=$(TargetPlatform) -packagesize=$(BuildSizeLimit) -build=&quot;$(ProjectOutputDirectory)/Packages/$(StagedPlatformFolder)/$(PackageConfiguration)&quot; -configuration=$(PackageConfiguration)" />
 				</ForEach>
 			</Node>
 			
@@ -762,9 +764,10 @@
 				<Property Name="TestPrerequisites" Value="$(PlatformToolsNodeName);$(PlatformEditorNodeName)" If="'$(GeneratedBuildVersion)' != '' and $(TestsRequireEditor)"/>
 				
 				<!-- Create the -build argument for Gauntlet based on whether builds are local or on the network-->
-				<Property Name="GauntletBuildArg" Value="-build=&quot;$(ProjectOutputDirectory)&quot;" />
+				<Property Name="GauntletBuildArg" Value="-build=&quot;$(ProjectOutputDirectory)/$(StagedPlatformFolder)&quot;" />
+				<Property Name="GauntletBuildArg" Value="-build=&quot;$(ProjectOutputDirectory)/Packages/$(StagedPlatformFolder)&quot; -packaged" If="$(UsePackaged)" />
 
-				<Do If="'$(NetworkOutputDirectory)' != '' Or '$(GeneratedBuildVersion)' != ''">
+				<Do If="!$(Horde) And ('$(NetworkOutputDirectory)' != '' Or '$(GeneratedBuildVersion)' != '')">
 					<Property Name="TestDirectory" Value="$(NetworkOutputDirectory)" />
 					<Property Name="TestDirectory" Value="$(GeneratedBuildDir)" If="'$(GeneratedBuildVersion)' != ''"/>
 					<Property Name="GauntletBuildArg" Value="-build=&quot;$(TestDirectory)&quot;" />
@@ -772,6 +775,21 @@
 					<Property Name="GauntletBuildArg" Value="$(GauntletBuildArg) -destlocalinstalldir=&quot;$(LocalInstallDirectory)&quot;" If="!$(IsLocalPlatform) And $(WithDeviceReservation)"/>
 				</Do>
 
+				<Do If="$(Horde)">
+					<Property Name="HordeArtifactPlatformRetrieveWithConfig" Value="$(HordeArtifactPlatformRetrieve) $(TestConfiguration) (Horde)"/>
+					<Node Name="$(HordeArtifactPlatformRetrieveWithConfig)" Requires="$(TestPrerequisites)">
+						<Switch>
+							<Case If="$(UsePackaged)">
+								<RetrieveArtifact Name="Packaged-$(StagedPlatformFolder)-$(TestConfiguration)" Type="packaged-build" OutputDir="$(ProjectOutputDirectory)/Packages/$(StagedPlatformFolder)/$(TestConfiguration)" />
+							</Case>
+							<Default>
+								<RetrieveArtifact Name="Staged-$(StagedPlatformFolder)" Type="staged-build" OutputDir="$(ProjectOutputDirectory)/$(StagedPlatformFolder)" />
+							</Default>
+						</Switch>
+					</Node>
+					<Property Name="TestPrerequisites" Value="$(TestPrerequisites);$(HordeArtifactPlatformRetrieveWithConfig)" />
+				</Do>
+
 				<Property Name="GauntletTestArgs" Value="$(GauntletTargetTestArgs)"/>
 				<Property Name="GauntletTestArgs" Value="$(GauntletTargetTestAltArgs)" If="ContainsItem('$(AlternateDevicePoolPlatforms)', '$(TargetPlatform)', ';')"/>
 
diff --git a/Engine/Build/Graph/Tasks/Inc/PathSettings.xml b/Engine/Build/Graph/Tasks/Inc/PathSettings.xml
--- a/Engine/Build/Graph/Tasks/Inc/PathSettings.xml
+++ b/Engine/Build/Graph/Tasks/Inc/PathSettings.xml
@@ -38,7 +38,7 @@
 	<Property Name="IsEpicBuildMachine" Value="true" If="$(IsBuildMachine) and $(IsEpicInternal)"/>
 
 	<!-- Publish packaged builds to the network ?-->
-	<Option Name="PublishBuilds" DefaultValue="$(IsEpicBuildMachine)" Description="Publish packaged builds to the builds folder on the network"/>
+	<Option Name="PublishBuilds" DefaultValue="false" Restrict="true|false" Description="Publish packaged builds to the builds folder on the network"/>
 
 	<!-- Publish a hashed staging data that developers can quickly sync to -->
 	<Option Name="PublishHashedStagingData" DefaultValue="false" Restrict="true|false" Description="Create hashes for syncing via Unsync" />
@@ -74,12 +74,12 @@
 
 	<!-- Declare the network outout path. Build machines need to write to the network for tests so use temp by default, but the publish path if that's being used -->
 	<Property Name="NetworkOutputDirectory" Value="" />
-	<Property Name="NetworkOutputDirectory" Value="$(NetworkTempRoot)" If="$(IsBuildMachine)"/>
-	<Property Name="NetworkOutputDirectory" Value="$(NetworkPublishRoot)" If="$(PublishBuilds) or $(PublishHashedStagingData)"/>
+	<Property Name="NetworkOutputDirectory" Value="$(NetworkTempRoot)" If="$(IsBuildMachine) And !$(Horde)"/>
+	<Property Name="NetworkOutputDirectory" Value="$(NetworkPublishRoot)" If="$(PublishBuilds)"/>
 	<Property Name="NetworkOutputDirectory" Value="$(NetworkOutputDirectory)/$(EscapedBranch)/$(BuildNamePath)" If="'$(NetworkOutputDirectory)' != ''"/>
 
 	<!-- Automation can publish reports to a folder that is mapped to a website -->
-	<Do If="$(IsEpicBuildMachine) or $(PublishReports)" >
+	<Do If="$(PublishReports)" >
 		<Property Name="AutomationReportOutputDirectory" Value="$(NetworkReportRoot)"/>
 		<Property Name="AutomationReportOutputDirectory" Value="$(AutomationReportOutputDirectory)/$(EscapedBranch)/$(ProjectName)/$(CL-String)"/>
 		<Property Name="AutomationReportUrl" Value="$(AutomationReportURLBase)/$(EscapedBranch)/$(ProjectName)/$(CL-String)"/>
diff --git a/Engine/Build/Graph/Tasks/Inc/PlatformOptions.xml b/Engine/Build/Graph/Tasks/Inc/PlatformOptions.xml
--- a/Engine/Build/Graph/Tasks/Inc/PlatformOptions.xml
+++ b/Engine/Build/Graph/Tasks/Inc/PlatformOptions.xml
@@ -50,12 +50,20 @@
 <Option Name="CookedEditorPlatformsToSkipTests" DefaultValue="" Restrict=".*" Description="List of editor cooked platforms that editor tests won't be run on" />
 <Option Name="PlatformsToSkipTests" DefaultValue="" Restrict=".*" Description="List of platforms that all tests won't be run on. Overrides Editor/TargetPlatformsToSkipTests when non-empty." />
 
-<Macro Name="AddPlatform" Arguments="Platform;IsConsole;IsEnabled" OptionalArguments="IsVirtualized">
+<!-- Packaging isn't supported by Gauntlet yet :( )-->
+<Property Name="PlatformsToSkipPackaging" Value="Win64+Linux" />
+
+<!-- These platforms can't run from loose staged content-->
+<Property Name="PlatformsThatRequirePackaging" Value="Android+IOS" />
+
+<Macro Name="AddPlatform" Arguments="Platform;IsConsole;IsEnabled" OptionalArguments="IsVirtualized;TestWithStageBuildOnly">
 	<Property Name="IsVirtualized" Value="('$(IsVirtualized)' != '' and '$(IsVirtualized)' != 'false')"/>
+	<Property Name="TestWithStageBuildOnly" Value="('$(TestWithStageBuildOnly)' != '' and '$(TestWithStageBuildOnly)' != 'false')"/>
 	<Property Name="AllTargetPlatforms" Value="$(AllTargetPlatforms)+$(Platform)" If="!ContainsItem('$(AllTargetPlatforms)', '$(Platform)', '+')"/>
 	<Property Name="AllConsolePlatforms" Value="$(AllConsolePlatforms)+$(Platform)" If="$(IsConsole) and !ContainsItem('$(AllConsolePlatforms)', '$(Platform)', '+')"/>
 	<Property Name="AllVirtualizedPlatforms" Value="$(AllVirtualizedPlatforms)+$(Platform)" If="$(IsVirtualized) And !ContainsItem('$(AllVirtualizedPlatforms)', '$(Platform)', '+')"/>
 	<Property Name="TargetPlatforms" Value="$(TargetPlatforms)+$(Platform)" If="$(IsEnabled) and !ContainsItem('$(TargetPlatforms)', '$(Platform)', '+')"/>
+	<Property Name="PlatformsToSkipPackaging" Value="$(PlatformsToSkipPackaging)+$(Platform)" If="$(TestWithStageBuildOnly) and !ContainsItem('$(PlatformsToSkipPackaging)', '$(Platform)', '+')"/>
 </Macro>
 	
 <!-- Include any scripts from platform extensions -->
@@ -159,12 +167,6 @@
 	</Do>
 </ForEach>
 
-<!-- Packaging isn't supported by Gauntlet yet :( )-->
-<Property Name="PlatformsToSkipPackaging" Value="Win64+Linux" />
-
-<!-- These platforms can't run from loose staged content-->
-<Property Name="PlatformsThatRequirePackaging" Value="Android+IOS" />
-
 <Option Name="WithDebugTargets" DefaultValue="false" Description="Build Debug configurations"/>
 <Option Name="WithDevelopmentTargets" DefaultValue="false" Description="Build Development configurations"/>
 <Option Name="WithTestTargets" DefaultValue="false" Description="Build Test configurations"/>
