Add metadata to exported artifacts
+ Create/Retrieve artifact for CookedEditor build
+ Add support to Gauntlet to retrieve artifact metadata

#rnx

diff --git a/Engine/Build/Graph/Tasks/BuildAndTestProject.xml b/Engine/Build/Graph/Tasks/BuildAndTestProject.xml
--- a/Engine/Build/Graph/Tasks/BuildAndTestProject.xml
+++ b/Engine/Build/Graph/Tasks/BuildAndTestProject.xml
@@ -158,6 +158,7 @@
 	<Include Script="Inc/CommonProperties.xml" />
 
 	<!-- Common settings but can be overidden -->
+	<Option Name="UseNetworkShare" DefaultValue="false" Restrict="true|false" Description="Whether to use network share for build output"/>
 	<Include Script="$(PathSettingsFile)" />
 	<Include Script="$(GauntletSettingsFile)" />
 	<Do If="'$(PreStepsExtendFile)' != ''">
@@ -240,6 +241,10 @@
 		</Do>
 	</Macro>
 	
+	<!-- Artifact metadata file -->
+	<Macro Name="WriteArtifactMetadataFile" Arguments="Type;Platform;Configuration;ArtifactDir">
+		<WriteTextFile Text="{&quot;type&quot;:&quot;$(Type)&quot;,&quot;branch&quot;:&quot;$(Branch)&quot;,&quot;platform&quot;:&quot;$(Platform)&quot;,&quot;change&quot;:&quot;$(Change)&quot;,&quot;preflight&quot;:&quot;$(PreflightChange)&quot;,&quot;buildName&quot;:&quot;$(BuildNamePath)&quot;,&quot;configuration&quot;:&quot;$(Configuration)&quot;}" File="$(ArtifactDir)/artifactmetadata.json" />
+	</Macro>
 
 	<!-- These properties will be filled in with nodes that are generated based on options -->
 	<Property Name="GeneratedToolNodes" Value="" />
@@ -412,11 +417,21 @@
 			</Node>
 
 			<!-- Publish CookedEditor -->
-			<Node Name="$(PreNodeName)Publish CookedEditor $(EditorPlatform)" Requires="$(PreNodeName)Stage CookedEditor $(EditorPlatform)">
-				<Log Message="Copying to $(NetworkOutputDirectory)/$(EditorPlatformPath)CookedEditor"/>
-				<Do If="!$(SkipPublish)">
-					<Copy Files="..." From="$(ProjectOutputDirectory)/$(EditorPlatformPath)CookedEditor" To="$(NetworkOutputDirectory)/$(EditorPlatformPath)CookedEditor"/>
-				</Do>
+			<Node Name="$(PlatformPublishNodeName)" Requires="$(PreNodeName)Stage CookedEditor $(EditorPlatform)">
+				<Switch>
+					<Case If="$(Horde)">
+						<Expand Name="WriteArtifactMetadataFile" Type="staged-build" Platform="$(EditorPlatform)" Configuration="Development" ArtifactDir="$(ProjectOutputDirectory)/$(EditorPlatformPath)CookedEditor" />
+						<Tag Files="$(ProjectOutputDirectory)/$(EditorPlatformPath)CookedEditor/..." With="#Staged CookedEditor Files"/>
+
+						<CreateArtifact Name="Staged-Cooked-Editor-$(EditorPlatform)" Type="staged-build" BaseDir="$(ProjectOutputDirectory)/$(EditorPlatformPath)CookedEditor" Keys="Project=$(ProjectName)" Metadata="Stream=$(Branch);Platform=$(EditorPlatform);Configuration=Development;BuildName=$(BuildNamePath);dashboard-category=Development Builds" Files="#Staged CookedEditor Files"/>
+					</Case>
+					<Default>
+						<Log Message="Copying to $(NetworkOutputDirectory)/$(EditorPlatformPath)CookedEditor"/>
+						<Do If="!$(SkipPublish)">
+							<Copy Files="..." From="$(ProjectOutputDirectory)/$(EditorPlatformPath)CookedEditor" To="$(NetworkOutputDirectory)/$(EditorPlatformPath)CookedEditor"/>
+						</Do>
+					</Default>
+				</Switch>
 			</Node>
 
 			<!-- Aggregate Nodes requirements -->
@@ -428,9 +443,17 @@
 		<!-- Tests -->
 		<Agent Name="Test $(ProjectName)CookedEditor $(EditorPlatform)" Type="$(TestAgentType)">
 			<Property Name="CustomModuleRoleArg" Value="-CustomModuleRoles=$(CookedEditorTargetName):CookedEditor"/>
+			<Property Name="CookedEditorBuildPath" Value="$(NetworkOutputDirectory)/$(EditorPlatformPath)CookedEditor" />
 			<!-- Prerequisites for the editor test to run -->
 			<Property Name="TestPrerequisites" Value="$(GeneratedCookedEditorNodes)"/>
-			<Property Name="TestPrerequisites" Value="$(TestPrerequisites);$(PlatformPublishNodeName)" If="!$(PublishBuilds)"/>
+			<Do If="$(Horde)">
+				<Property Name="HordeArtifactPlatformRetrieve" Value="$(PreNodeName)Retrieve CookedEditor Build Artifact $(EditorPlatform) (Horde)"/>
+				<Node Name="$(HordeArtifactPlatformRetrieve)" Requires="$(TestPrerequisites);$(PlatformPublishNodeName)">
+					<RetrieveArtifact Name="Staged-Cooked-Editor-$(EditorPlatform)" Type="staged-build" OutputDir="$(ProjectOutputDirectory)/$(EditorPlatformPath)CookedEditor" />
+				</Node>
+				<Property Name="TestPrerequisites" Value="$(TestPrerequisites);$(HordeArtifactPlatformRetrieve)" />
+				<Property Name="CookedEditorBuildPath" Value="$(ProjectOutputDirectory)/$(EditorPlatformPath)CookedEditor" />
+			</Do>
 			<!-- Declare a boot test for the editor-->
 			<Property Name="PlatformBootTestNodeName" Value="$(PreNodeName)CookedEditor $(EditorPlatform) Test=BootTest" />
 			<Do If="'$(AutomationReportOutputDirectory)' != ''">
@@ -442,7 +465,7 @@
 			<Property Name="NodeAnnotations" Value=""/>
 			<Property Name="NodeAnnotations" Value="Workflow=$(AnnotationsTarget)" If="ContainsItem('$(AnnotationsPlatforms)', '$(EditorPlatform)Editor', ';') and ContainsItem('$(AnnotationsTests)', 'BootTest', ';')"/>
 			<Node Name="$(PlatformBootTestNodeName)" Requires="$(TestPrerequisites)" Annotations="$(NodeAnnotations)">
-				<Command Name="RunUnreal" Arguments="-project=$(TargetProject) -build=&quot;$(NetworkOutputDirectory)/$(EditorPlatformPath)CookedEditor&quot; -test=&quot;UE.EditorBootTest&quot; $(GauntletCookedEditorTestArgs) $(CustomModuleRoleArg) $(GauntletReportArgs)" />
+				<Command Name="RunUnreal" Arguments="-project=$(TargetProject) -build=&quot;$(CookedEditorBuildPath)&quot; -test=&quot;UE.EditorBootTest&quot; $(GauntletCookedEditorTestArgs) $(CustomModuleRoleArg) $(GauntletReportArgs)" />
 			</Node>
 
 			<!-- Link the boot test in to the prerequisites and generated nodes if its enabled-->
@@ -467,7 +490,7 @@
 				<Property Name="TestArgumentsWithinParentheses" Value=""/>
 				<Expand Name="SplitCombinedTestNameAndArgs"/>
 				<Node Name="$(TestNodeName)" Requires="$(TestPrerequisites)" Annotations="$(NodeAnnotations)">
-					<Command Name="RunUnreal" Arguments="-project=$(TargetProject) -build=&quot;$(NetworkOutputDirectory)/$(EditorPlatformPath)CookedEditor&quot; -test=&quot;$(TestNameBeforeParentheses)&quot; $(GauntletReportArgs) $(CustomModuleRoleArg) $(GauntletCookedEditorTestArgs) $(TestArgumentsWithinParentheses)" />
+					<Command Name="RunUnreal" Arguments="-project=$(TargetProject) -build=&quot;$(CookedEditorBuildPath)&quot; -test=&quot;$(TestNameBeforeParentheses)&quot; $(GauntletReportArgs) $(CustomModuleRoleArg) $(GauntletCookedEditorTestArgs) $(TestArgumentsWithinParentheses)" />
 				</Node>
 				<!-- Add each test node to our generated list -->
 				<Property Name="PlatformCookedEditorTestNodes" Value="$(PlatformCookedEditorTestNodes);$(TestNodeName)" />
@@ -640,6 +663,7 @@
 						<Command Name="BuildCookRun" Arguments="$(BCRArgs) -configuration=$(PackageConfigurations) $(StageAndPackageArguments) -skipbuild -skipcook -stage -pak" />
 					</Default>
 				</Switch>
+				<Expand Name="WriteArtifactMetadataFile" Type="staged-build" Platform="$(StagedPlatformFolder)" Configuration="$(PackageConfigurations)" ArtifactDir="$(ProjectOutputDirectory)/$(StagedPlatformFolder)" />
 			</Node>
 			
 			<!-- Hash -->
@@ -731,6 +755,7 @@
 				<Node Name="$(HordeArtifactPlatformPackagedNodeName)" Requires="$(PlatformPackageNodeName)">
 					<ForEach Name="PackageConfiguration" Values="$(PackageConfigurations)" Separator="+">
 						<Command Name="BuildCookRun" Arguments="$(BCRArgs) -configuration=$(PackageConfiguration) $(StageAndPackageArguments) -skipbuild -skipcook -skippak $(HordePackageSkipBCRArguments) -archive -archivedirectory=&quot;$(ProjectOutputDirectory)/Packages/$(StagedPlatformFolder)/$(PackageConfiguration)&quot;" />
+						<Expand Name="WriteArtifactMetadataFile" Type="packaged-build" Platform="$(StagedPlatformFolder)" Configuration="$(PackageConfiguration)" ArtifactDir="$(ProjectOutputDirectory)/Packages/$(StagedPlatformFolder)/$(PackageConfiguration)" />
 						<Tag Files="$(ProjectOutputDirectory)/Packages/$(StagedPlatformFolder)/$(PackageConfiguration)/..." With="#$(PlatformPackageNodeName) $(PackageConfiguration) Files"/>
 
 						<CreateArtifact Name="Packaged-$(StagedPlatformFolder)-$(PackageConfiguration)" Type="packaged-build" BaseDir="$(ProjectOutputDirectory)/Packages/$(StagedPlatformFolder)/$(PackageConfiguration)" Keys="Project=$(ProjectName)" Metadata="Stream=$(Branch);Platform=$(StagedPlatformFolder);Configuration=$(PackageConfiguration);BuildName=$(BuildNamePath);dashboard-category=$(PackageConfiguration) Builds" Files="#$(PlatformPackageNodeName) $(PackageConfiguration) Files"/>
diff --git a/Engine/Build/Graph/Tasks/Inc/PathSettings.xml b/Engine/Build/Graph/Tasks/Inc/PathSettings.xml
--- a/Engine/Build/Graph/Tasks/Inc/PathSettings.xml
+++ b/Engine/Build/Graph/Tasks/Inc/PathSettings.xml
@@ -38,7 +38,8 @@
 	<Property Name="IsEpicBuildMachine" Value="true" If="$(IsBuildMachine) and $(IsEpicInternal)"/>
 
 	<!-- Publish packaged builds to the network ?-->
-	<Option Name="PublishBuilds" DefaultValue="false" Restrict="true|false" Description="Publish packaged builds to the builds folder on the network"/>
+	<Option Name="UseNetworkShare" DefaultValue="$(IsBuildMachine)" Restrict="true|false" Description="Whether to use network share for build output"/>
+	<Option Name="PublishBuilds" DefaultValue="$(UseNetworkShare)" Restrict="true|false" Description="Publish packaged builds to the builds folder on the network"/>
 
 	<!-- Publish a hashed staging data that developers can quickly sync to -->
 	<Option Name="PublishHashedStagingData" DefaultValue="false" Restrict="true|false" Description="Create hashes for syncing via Unsync" />
@@ -74,8 +75,8 @@
 
 	<!-- Declare the network outout path. Build machines need to write to the network for tests so use temp by default, but the publish path if that's being used -->
 	<Property Name="NetworkOutputDirectory" Value="" />
-	<Property Name="NetworkOutputDirectory" Value="$(NetworkTempRoot)" If="$(IsBuildMachine) And !$(Horde)"/>
-	<Property Name="NetworkOutputDirectory" Value="$(NetworkPublishRoot)" If="$(PublishBuilds)"/>
+	<Property Name="NetworkOutputDirectory" Value="$(NetworkTempRoot)" If="$(UseNetworkShare)"/>
+	<Property Name="NetworkOutputDirectory" Value="$(NetworkPublishRoot)" If="$(PublishBuilds) and $(UseNetworkShare)"/>
 	<Property Name="NetworkOutputDirectory" Value="$(NetworkOutputDirectory)/$(EscapedBranch)/$(BuildNamePath)" If="'$(NetworkOutputDirectory)' != ''"/>
 
 	<!-- Automation can publish reports to a folder that is mapped to a website -->
diff --git a/Engine/Source/Programs/AutomationTool/Gauntlet/Unreal/BuildSource/Gauntlet.UnrealBuildSource.cs b/Engine/Source/Programs/AutomationTool/Gauntlet/Unreal/BuildSource/Gauntlet.UnrealBuildSource.cs
--- a/Engine/Source/Programs/AutomationTool/Gauntlet/Unreal/BuildSource/Gauntlet.UnrealBuildSource.cs
+++ b/Engine/Source/Programs/AutomationTool/Gauntlet/Unreal/BuildSource/Gauntlet.UnrealBuildSource.cs
@@ -8,6 +8,8 @@ using System.Text.RegularExpressions;
 using System.IO;
 using System.Linq;
 using EpicGames.Core;
+using System.Text.Json.Serialization;
+using System.Text.Json;
 
 namespace Gauntlet
 {
@@ -29,10 +31,41 @@ namespace Gauntlet
 
 		public int Changelist { get; protected set; }
 
+		public int Preflight { get; protected set; }
+
 		public bool EditorValid { get; protected set; }
 
 		protected Dictionary<UnrealTargetPlatform, List<IBuild>> DiscoveredBuilds;
 
+		protected class ArtifactMetadata
+		{
+			[JsonPropertyName("change")]
+			public string Change { get; set; }
+			[JsonPropertyName("preflight")]
+			public string Preflight { get; set; }
+			[JsonPropertyName("branch")]
+			public string Branch { get; set; }
+			[JsonPropertyName("buildName")]
+			public string BuildName { get; set; }
+			[JsonPropertyName("configuration")]
+			public string Configuration { get; set; }
+			[JsonPropertyName("platform")]
+			public string Platform { get; set; }
+			[JsonPropertyName("type")]
+			public string Type { get; set; }
+
+			public static ArtifactMetadata LoadFromJson(string FilePath)
+			{
+				JsonSerializerOptions Options = new JsonSerializerOptions
+				{
+					PropertyNameCaseInsensitive = true
+				};
+				string JsonString = File.ReadAllText(FilePath);
+				ArtifactMetadata JsonTestPassResults = JsonSerializer.Deserialize<ArtifactMetadata>(JsonString, Options);
+				return JsonTestPassResults;
+			}
+		}
+
 		public UnrealBuildSource(string InProjectName, FileReference InProjectPath, DirectoryReference InUnrealPath, bool InUsesSharedBuildType, string BuildReference) 
 		{
 			InitBuildSource(InProjectName, InProjectPath, InUnrealPath, InUsesSharedBuildType, BuildReference, null);
@@ -105,9 +138,56 @@ namespace Gauntlet
 				Changelist = 0;
 			}
 
+			// Preflight?
+			M = Regex.Match(BuildName, @"-PF-(\d+)");
+			Preflight = M.Success? Convert.ToInt32(M.Groups[1].Value) : 0;
+
+			// Look if the build has an artifact metadata file
+			string ArtifactMetadataFilePath = Path.Combine(InBuildArgument, "artifactmetadata.json");
+			if (File.Exists(ArtifactMetadataFilePath))
+			{
+				ArtifactMetadata Metadata = null;
+				try
+				{
+					Metadata = ArtifactMetadata.LoadFromJson(ArtifactMetadataFilePath);
+				}
+				catch (Exception Ex)
+				{
+					Log.Warning("Could not parse ArtifactMetadata file '{File}': {Exception}", ArtifactMetadataFilePath, Ex.Message);
+				}
+				if (Metadata != null)
+				{
+					if (!string.IsNullOrEmpty(Metadata.Branch))
+					{
+						Branch = Metadata.Branch;
+					}
+					if (!string.IsNullOrEmpty(Metadata.Change))
+					{
+						int MetaChange;
+						if (int.TryParse(Metadata.Change, out MetaChange))
+						{
+							Changelist = MetaChange;
+						}
+					}
+					if (!string.IsNullOrEmpty(Metadata.Preflight))
+					{
+						int MetaPreflight;
+						if (int.TryParse(Metadata.Preflight, out MetaPreflight))
+						{
+							Preflight = MetaPreflight;
+						}
+					}
+					if (!string.IsNullOrEmpty(Metadata.BuildName))
+					{
+						BuildName = Metadata.BuildName;
+					}
+				}
+			}
+
 			// allow user overrides (TODO - centralize all this!)
 			Branch = Globals.Params.ParseValue("branch", Branch);
 			Changelist = Convert.ToInt32(Globals.Params.ParseValue("changelist", Changelist.ToString()));
+			Preflight = Convert.ToInt32(Globals.Params.ParseValue("preflight", Preflight.ToString()));
 
 			// We resolve these on demand
 			DiscoveredBuilds = new Dictionary<UnrealTargetPlatform, List<IBuild>>();
