Fix overriding artifact file error + enforce project name in artifact name

#rnx

diff --git a/Engine/Build/Graph/Tasks/BuildAndTestProject.xml b/Engine/Build/Graph/Tasks/BuildAndTestProject.xml
--- a/Engine/Build/Graph/Tasks/BuildAndTestProject.xml
+++ b/Engine/Build/Graph/Tasks/BuildAndTestProject.xml
@@ -423,7 +423,7 @@
 						<Expand Name="WriteArtifactMetadataFile" Type="staged-build" Platform="$(EditorPlatform)" Configuration="Development" ArtifactDir="$(ProjectOutputDirectory)/$(EditorPlatformPath)CookedEditor" />
 						<Tag Files="$(ProjectOutputDirectory)/$(EditorPlatformPath)CookedEditor/..." With="#Staged CookedEditor Files"/>
 
-						<CreateArtifact Name="Staged-Cooked-Editor-$(EditorPlatform)" Type="staged-build" BaseDir="$(ProjectOutputDirectory)/$(EditorPlatformPath)CookedEditor" Keys="Project=$(ProjectName)" Metadata="Stream=$(Branch);Platform=$(EditorPlatform);Configuration=Development;BuildName=$(BuildNamePath);dashboard-category=Development Builds" Files="#Staged CookedEditor Files"/>
+						<CreateArtifact Name="$(ProjectName)-Staged-Cooked-Editor-$(EditorPlatform)" Type="staged-build" BaseDir="$(ProjectOutputDirectory)/$(EditorPlatformPath)CookedEditor" Keys="Project=$(ProjectName)" Metadata="Stream=$(Branch);Platform=$(EditorPlatform);Configuration=Development;BuildName=$(BuildNamePath);dashboard-category=Development Builds" Files="#Staged CookedEditor Files"/>
 					</Case>
 					<Default>
 						<Log Message="Copying to $(NetworkOutputDirectory)/$(EditorPlatformPath)CookedEditor"/>
@@ -449,7 +449,7 @@
 			<Do If="$(Horde)">
 				<Property Name="HordeArtifactPlatformRetrieve" Value="$(PreNodeName)Retrieve CookedEditor Build Artifact $(EditorPlatform) (Horde)"/>
 				<Node Name="$(HordeArtifactPlatformRetrieve)" Requires="$(TestPrerequisites);$(PlatformPublishNodeName)">
-					<RetrieveArtifact Name="Staged-Cooked-Editor-$(EditorPlatform)" Type="staged-build" OutputDir="$(ProjectOutputDirectory)/$(EditorPlatformPath)CookedEditor" />
+					<RetrieveArtifact Name="$(ProjectName)-Staged-Cooked-Editor-$(EditorPlatform)" Type="staged-build" OutputDir="$(ProjectOutputDirectory)/$(EditorPlatformPath)CookedEditor" />
 				</Node>
 				<Property Name="TestPrerequisites" Value="$(TestPrerequisites);$(HordeArtifactPlatformRetrieve)" />
 				<Property Name="CookedEditorBuildPath" Value="$(ProjectOutputDirectory)/$(EditorPlatformPath)CookedEditor" />
@@ -663,7 +663,6 @@
 						<Command Name="BuildCookRun" Arguments="$(BCRArgs) -configuration=$(PackageConfigurations) $(StageAndPackageArguments) -skipbuild -skipcook -stage -pak" />
 					</Default>
 				</Switch>
-				<Expand Name="WriteArtifactMetadataFile" Type="staged-build" Platform="$(StagedPlatformFolder)" Configuration="$(PackageConfigurations)" ArtifactDir="$(ProjectOutputDirectory)/$(StagedPlatformFolder)" />
 			</Node>
 			
 			<!-- Hash -->
@@ -711,7 +710,7 @@
 					</Do>
 					<Tag Files="$(SnapshotLocalDir)/..." With="#SnapshotDescriptor_$(StagedPlatformFolder)"/>
 
-					<CreateArtifact Name="ci-snapshot-$(StagedPlatformFolder)" Description="Snapshot ($(StagedPlatformFolder))" Type="ugs-snapshot" BaseDir="Snapshots" Keys="ugs-project=$(TargetProject)" Metadata="ArchiveType=SnapshotDescriptor_$(StagedPlatformFolder);dashboard-category=Snapshots" Files="#SnapshotDescriptor_$(StagedPlatformFolder)"/>
+					<CreateArtifact Name="$(ProjectName)-ci-snapshot-$(StagedPlatformFolder)" Description="Snapshot ($(StagedPlatformFolder))" Type="ugs-snapshot" BaseDir="Snapshots" Keys="ugs-project=$(TargetProject)" Metadata="ArchiveType=SnapshotDescriptor_$(StagedPlatformFolder);dashboard-category=Snapshots" Files="#SnapshotDescriptor_$(StagedPlatformFolder)"/>
 				</Node>
 				<Property Name="PlatformPublishStageDependencyName" Value="$(PlatformSnapshotPublishNodeName)" />
 			</Do>
@@ -726,9 +725,10 @@
 			<!-- Publish (Staged) (Horde) -->
 			<Do If="$(Horde)">
 				<Node Name="$(PlatformStagedArtifactNodeName)" Requires="$(PlatformStageNodeName)">
+					<Expand Name="WriteArtifactMetadataFile" Type="staged-build" Platform="$(StagedPlatformFolder)" Configuration="$(PackageConfigurations)" ArtifactDir="$(ProjectOutputDirectory)/$(StagedPlatformFolder)" />
 					<Tag Files="$(ProjectOutputDirectory)/$(StagedPlatformFolder)/..." With="#$(PlatformStageNodeName) Files"/>
 
-					<CreateArtifact Name="Staged-$(StagedPlatformFolder)" Type="staged-build" BaseDir="$(ProjectOutputDirectory)/$(StagedPlatformFolder)" Keys="Project=$(ProjectName)" Metadata="Stream=$(Branch);Platform=$(StagedPlatformFolder);BuildName=$(BuildNamePath);dashboard-category=Staged Builds" Files="#$(PlatformStageNodeName) Files"/>
+					<CreateArtifact Name="$(ProjectName)-Staged-$(StagedPlatformFolder)" Type="staged-build" BaseDir="$(ProjectOutputDirectory)/$(StagedPlatformFolder)" Keys="Project=$(ProjectName)" Metadata="Stream=$(Branch);Platform=$(StagedPlatformFolder);BuildName=$(BuildNamePath);dashboard-category=Staged Builds" Files="#$(PlatformStageNodeName) Files"/>
 				</Node>
 			</Do>
 
@@ -758,7 +758,7 @@
 						<Expand Name="WriteArtifactMetadataFile" Type="packaged-build" Platform="$(StagedPlatformFolder)" Configuration="$(PackageConfiguration)" ArtifactDir="$(ProjectOutputDirectory)/Packages/$(StagedPlatformFolder)/$(PackageConfiguration)" />
 						<Tag Files="$(ProjectOutputDirectory)/Packages/$(StagedPlatformFolder)/$(PackageConfiguration)/..." With="#$(PlatformPackageNodeName) $(PackageConfiguration) Files"/>
 
-						<CreateArtifact Name="Packaged-$(StagedPlatformFolder)-$(PackageConfiguration)" Type="packaged-build" BaseDir="$(ProjectOutputDirectory)/Packages/$(StagedPlatformFolder)/$(PackageConfiguration)" Keys="Project=$(ProjectName)" Metadata="Stream=$(Branch);Platform=$(StagedPlatformFolder);Configuration=$(PackageConfiguration);BuildName=$(BuildNamePath);dashboard-category=$(PackageConfiguration) Builds" Files="#$(PlatformPackageNodeName) $(PackageConfiguration) Files"/>
+						<CreateArtifact Name="$(ProjectName)-Packaged-$(StagedPlatformFolder)-$(PackageConfiguration)" Type="packaged-build" BaseDir="$(ProjectOutputDirectory)/Packages/$(StagedPlatformFolder)/$(PackageConfiguration)" Keys="Project=$(ProjectName)" Metadata="Stream=$(Branch);Platform=$(StagedPlatformFolder);Configuration=$(PackageConfiguration);BuildName=$(BuildNamePath);dashboard-category=$(PackageConfiguration) Builds" Files="#$(PlatformPackageNodeName) $(PackageConfiguration) Files"/>
 					</ForEach>
 				</Node>
 			</Do>
@@ -805,10 +805,10 @@
 					<Node Name="$(HordeArtifactPlatformRetrieveWithConfig)" Requires="$(TestPrerequisites)">
 						<Switch>
 							<Case If="$(UsePackaged)">
-								<RetrieveArtifact Name="Packaged-$(StagedPlatformFolder)-$(TestConfiguration)" Type="packaged-build" OutputDir="$(ProjectOutputDirectory)/Packages/$(StagedPlatformFolder)/$(TestConfiguration)" />
+								<RetrieveArtifact Name="$(ProjectName)-Packaged-$(StagedPlatformFolder)-$(TestConfiguration)" Type="packaged-build" OutputDir="$(ProjectOutputDirectory)/Packages/$(StagedPlatformFolder)/$(TestConfiguration)" />
 							</Case>
 							<Default>
-								<RetrieveArtifact Name="Staged-$(StagedPlatformFolder)" Type="staged-build" OutputDir="$(ProjectOutputDirectory)/$(StagedPlatformFolder)" />
+								<RetrieveArtifact Name="$(ProjectName)-Staged-$(StagedPlatformFolder)" Type="staged-build" OutputDir="$(ProjectOutputDirectory)/$(StagedPlatformFolder)" />
 							</Default>
 						</Switch>
 					</Node>
