Re-enable publish build to netshare by default

#rnx

diff --git a/Engine/Build/Graph/Tasks/BuildAndTestProject.xml b/Engine/Build/Graph/Tasks/BuildAndTestProject.xml
--- a/Engine/Build/Graph/Tasks/BuildAndTestProject.xml
+++ b/Engine/Build/Graph/Tasks/BuildAndTestProject.xml
@@ -158,7 +158,7 @@
 	<Include Script="Inc/CommonProperties.xml" />
 
 	<!-- Common settings but can be overidden -->
-	<Option Name="UseNetworkShare" DefaultValue="false" Restrict="true|false" Description="Whether to use network share for build output"/>
+	<Option Name="UseNetworkShare" DefaultValue="$(IsBuildMachine)" Restrict="true|false" Description="Whether to use network share for build output"/>
 	<Include Script="$(PathSettingsFile)" />
 	<Include Script="$(GauntletSettingsFile)" />
 	<Do If="'$(PreStepsExtendFile)' != ''">
