Allow buildversion to be input through option

#rnx

diff --git a/Engine/Build/Graph/Tasks/Inc/CommonProperties.xml b/Engine/Build/Graph/Tasks/Inc/CommonProperties.xml
--- a/Engine/Build/Graph/Tasks/Inc/CommonProperties.xml
+++ b/Engine/Build/Graph/Tasks/Inc/CommonProperties.xml
@@ -19,7 +19,7 @@
 	<Property Name="PreflightSuffix" Value="-PF-$(PreflightChange)-$(UE_HORDE_JOBID:-6)" If="$(IsPreflight)"/>
 	<Property Name="CL-String" Value="CL-$(Change)$(PreflightSuffix)"/>
 	<Property Name="BuildNamePath" Value="$(ProjectName)-$(CL-String)"/>
-	<Property Name="BuildVersion" Value="$(ProjectName)-CL-$(Change)"/>
+	<Option Name="BuildVersion" DefaultValue="$(ProjectName)-CL-$(Change)" Description="Value to set as BuildVersion during compilation"/>
 
 	<Include Script="SharedProperties.xml" />
 
