Change IterativeCompare to allow it to be split into multiple nodes to allow tracking of durations of each step over time.  Also added ability to do a Null Cook in between incremental and Full Cook.  Added CookProcessCount option.

diff --git a/Engine/Build/Graph/Tests/IterativeCompare.xml b/Engine/Build/Graph/Tests/IterativeCompare.xml
--- a/Engine/Build/Graph/Tests/IterativeCompare.xml
+++ b/Engine/Build/Graph/Tests/IterativeCompare.xml
@@ -24,27 +24,41 @@
 	<Property Name="UBACompileArgs" Value="-UBA -ActionStallReportTime=300 -UBALinkRemote" If="$(WithUBA)"/>
 	
 	<!-- Agent=IterativeCompareAgent variables -->
+	<Option Name="CookProcessCount" DefaultValue="" Description="Number of processes to use during cooking, default is ini-driven."/>
+	<Option Name="WithNullCook" Restrict="true|false" DefaultValue="false" Description="Whether to do a null cook or not"/>
+	<Option Name="NullCookCustomConfig" DefaultValue="" Description="What custom configuration to apply if doing a null cook (see -CustomConfig commandline argument)"/>
 	<Property Name="FullPathProjectDir" Value="$(RootDir)/$(ProjectDir)"/>
 	<Property Name="EditorPlatform" Value="Win64"/>
 	<Property Name="CommonCookArgs" Value="-targetplatform=$(TargetPlatform) -zenstore -cookincremental"/>
-	<Property Name="IterativeCompareDescriptiveName" Value="IterativeCompare $(EditorPlatform) $(ProjectName)"/>
+	<Property Name="CommonCookArgs" Value="$(CommonCookArgs) -CookProcessCount=$(CookProcessCount)" If="'$(CookProcessCount)' != ''"/>
 	<Property Name="PlatformCookedDir" Value="$(FullPathProjectDir)/Saved/Cooked/$(TargetPlatform)"/>
 	<Property Name="CompareDirIncremental" Value="$(FullPathProjectDir)/Intermediate/IterativeCompare/Iter/$(TargetPlatform)"/>
 	<Property Name="CompareDirRecook" Value="$(FullPathProjectDir)/Intermediate/IterativeCompare/Recook/$(TargetPlatform)"/>
 	<Property Name="NewLine" Value="&#10;"/>
 
+	<Property Name="CompileNodeName" Value="Compile $(EditorPlatform) $(ProjectName)"/>
+	<Property Name="NaturalIncrementalCookNodeName" Value="Natural Incremental Cook $(ProjectName) $(TargetPlatform)"/>
+	<Property Name="StoreIncrementalCookNodeName" Value="Store Incremental Cook $(ProjectName) $(TargetPlatform)"/>
+	<Property Name="NullIncrementalCookNodeName" Value="Null Incremental Cook $(ProjectName) $(TargetPlatform)"/>
+	<Property Name="CleanNullIncrementalCookNodeName" Value="Clean Null Incremental Cook $(ProjectName) $(TargetPlatform)"/>
+	<Property Name="FullCookNodeName" Value="Full Cook $(ProjectName) $(TargetPlatform)"/>
+	<Property Name="StoreFullCookNodeName" Value="Store Full Cook $(ProjectName) $(TargetPlatform)"/>
+	<Property Name="DiffCooksNodeName" Value="Diff Cooks $(ProjectName) $(TargetPlatform)"/>
+
 	<Agent Name="IterativeCompareAgent" Type="$(AgentType)">
-		<Node Name="$(IterativeCompareDescriptiveName)">
+		<Node Name="$(CompileNodeName)">
 			<!-- Compile. -->
 			<Log Message="$(NewLine)========== IterativeCompare.xml: Compile =========="/>
 			<Compile Target="ShaderCompileWorker" Platform="$(EditorPlatform)" Project="$(ProjectFile)" Configuration="Development" Arguments="$(UBACompileArgs)"/>
 			<Compile Target="ZenLaunch" Platform="$(EditorPlatform)" Configuration="Development" Arguments="$(UBACompileArgs)"/>
 			<Compile Target="$(EditorTarget)" Platform="$(EditorPlatform)" Project="$(ProjectFile)" Configuration="Development" Arguments="$(UBACompileArgs)"/>
-
+		</Node>
+		<Node Name="$(NaturalIncrementalCookNodeName)" Requires="$(CompileNodeName)">
 			<!-- Execute the incremental cook (we don't pass -forcerecook). -->
 			<Log Message="$(NewLine)========== IterativeCompare.xml: Incremental Cook =========="/>
 			<Commandlet Name="Cook" Project="$(ProjectFile)" Arguments="$(CommonCookArgs)"/>
-
+		</Node>
+		<Node Name="$(StoreIncrementalCookNodeName)" Requires="$(NaturalIncrementalCookNodeName)">
 			<!-- Move all cook artifacts to the compare_incremental dir. -->
 			<Log Message="$(NewLine)========== IterativeCompare.xml: Move =========="/>
 			<Log Message="Delete Directories=$(CompareDirIncremental)"/>
@@ -58,11 +72,30 @@
 			<Log Message="$(NewLine)========== IterativeCompare.xml: ZenMirror Incremental Cook =========="/>
 			<Log Message="ZenMirror Project=$(ProjectFile) Platform=$(TargetPlatform) DestinationFileDir=$(CompareDirIncremental)"/>
 			<ZenMirror Project="$(ProjectFile)" Platform="$(TargetPlatform)" DestinationFileDir="$(CompareDirIncremental)"/>
-
+		</Node>
+		<Property Name="FullCookPrequisiteNodeName" Value="$(StoreIncrementalCookNodeName)"/>
+		<Do If="$(WithNullCook)">
+			<Node Name="$(NullIncrementalCookNodeName)" Requires="$(StoreIncrementalCookNodeName)">
+				<!-- Execute the null incremental cook. -->
+				<Log Message="$(NewLine)========== IterativeCompare.xml: Null Incremental Cook =========="/>
+				<Property Name="NullCookArgs" Value="$(CommonCookArgs)"/>
+				<Property Name="NullCookArgs" Value="$(NullCookArgs) -CustomConfig=$(NullCookCustomConfig)" If="'$(NullCookCustomConfig)' != ''"/>
+				<Commandlet Name="Cook" Project="$(ProjectFile)" Arguments="$(NullCookArgs)"/>
+			</Node>
+			<Node Name="$(CleanNullIncrementalCookNodeName)" Requires="$(NullIncrementalCookNodeName)">
+				<!-- Move all cook artifacts to the compare_incremental dir. -->
+				<Log Message="$(NewLine)========== IterativeCompare.xml: Clean Null =========="/>
+				<Log Message="Delete Directories=$(PlatformCookedDir)"/>
+				<Delete Directories="$(PlatformCookedDir)"/>
+			</Node>
+			<Property Name="FullCookPrequisiteNodeName" Value="$(CleanNullIncrementalCookNodeName)"/>
+		</Do>
+		<Node Name="$(FullCookNodeName)" Requires="$(FullCookPrequisiteNodeName)">
 			<!-- Execute the forcerecook cook. -->
 			<Log Message="$(NewLine)========== IterativeCompare.xml: Non-Incremental Cook =========="/>
 			<Commandlet Name="Cook" Project="$(ProjectFile)" Arguments="$(CommonCookArgs) -forcerecook"/>
-
+		</Node>
+		<Node Name="$(StoreFullCookNodeName)" Requires="$(FullCookNodeName)">
 			<!-- Copy all cook artifacts to the compare_incremental dir. Copy rather than move so they remain
 			     behind for the next instance of this job to use in an incremental cook. -->
 			<Log Message="$(NewLine)========== IterativeCompare.xml: Copy =========="/>
@@ -75,13 +108,14 @@
 			<Log Message="$(NewLine)========== IterativeCompare.xml: ZenMirror ForceRecook Cook =========="/>
 			<Log Message="ZenMirror Project=$(ProjectFile) Platform=$(TargetPlatform) DestinationFileDir=$(CompareDirRecook)"/>
 			<ZenMirror Project="$(ProjectFile)" Platform="$(TargetPlatform)" DestinationFileDir="$(CompareDirRecook)"/>
-
+		</Node>
+		<Node Name="$(DiffCooksNodeName)" Requires="$(StoreFullCookNodeName)">
 			<!-- Compare the incremental cook with the non-incremental cook. -->
 			<Log Message="$(NewLine)========== IterativeCompare.xml: DiffCook =========="/>
 			<Commandlet Name="DiffCook" Project="$(ProjectFile)" Arguments="-base=$(CompareDirRecook) -comp=$(CompareDirIncremental) -targetplatform=$(TargetPlatform) -addedverbosity=error -removedverbosity=error -modifiedverbosity=warning -hide=modifiedpackages"/>
 		</Node>
 	</Agent>
 
-	<Aggregate Name="IterativeCompare" Requires="$(IterativeCompareDescriptiveName)"/>
+	<Aggregate Name="IterativeCompare" Requires="$(DiffCooksNodeName)"/>
 
 </BuildGraph>
\ No newline at end of file
