Ensure we don't delete cooked platlform directory before performing a null cook in the IncrementalCompare job.

diff --git a/Engine/Build/Graph/Tests/IterativeCompare.xml b/Engine/Build/Graph/Tests/IterativeCompare.xml
--- a/Engine/Build/Graph/Tests/IterativeCompare.xml
+++ b/Engine/Build/Graph/Tests/IterativeCompare.xml
@@ -65,8 +65,10 @@
 			<Delete Directories="$(CompareDirIncremental)"/>
 			<Log Message="Move From=$(PlatformCookedDir)/*.* To=$(CompareDirIncremental)"/>
 			<Move From="$(PlatformCookedDir)/*.*" To="$(CompareDirIncremental)" ErrorIfNotFound="true"/>
-			<Log Message="Delete Directories=$(PlatformCookedDir)"/>
-			<Delete Directories="$(PlatformCookedDir)"/>
+			<Do If="!$(WithNullCook)">
+				<Log Message="Delete Directories=$(PlatformCookedDir)"/>
+				<Delete Directories="$(PlatformCookedDir)"/>
+			</Do>
 
 			<!-- Write the files from incremental cook's zenstore out into CompareDirIncremental. -->
 			<Log Message="$(NewLine)========== IterativeCompare.xml: ZenMirror Incremental Cook =========="/>
