Don't move files out of the cooked output directory if doing a null cook after incremental.  Copy them instead.

diff --git a/Engine/Build/Graph/Tests/IterativeCompare.xml b/Engine/Build/Graph/Tests/IterativeCompare.xml
--- a/Engine/Build/Graph/Tests/IterativeCompare.xml
+++ b/Engine/Build/Graph/Tests/IterativeCompare.xml
@@ -63,9 +63,13 @@
 			<Log Message="$(NewLine)========== IterativeCompare.xml: Move =========="/>
 			<Log Message="Delete Directories=$(CompareDirIncremental)"/>
 			<Delete Directories="$(CompareDirIncremental)"/>
-			<Log Message="Move From=$(PlatformCookedDir)/*.* To=$(CompareDirIncremental)"/>
-			<Move From="$(PlatformCookedDir)/*.*" To="$(CompareDirIncremental)" ErrorIfNotFound="true"/>
+			<Do If="$(WithNullCook)">
+				<Log Message="Copy From=$(PlatformCookedDir)/*.* To=$(CompareDirIncremental)"/>
+				<Copy From="$(PlatformCookedDir)/*.*" To="$(CompareDirIncremental)" ErrorIfNotFound="true"/>
+			</Do>
 			<Do If="!$(WithNullCook)">
+				<Log Message="Move From=$(PlatformCookedDir)/*.* To=$(CompareDirIncremental)"/>
+				<Move From="$(PlatformCookedDir)/*.*" To="$(CompareDirIncremental)" ErrorIfNotFound="true"/>
 				<Log Message="Delete Directories=$(PlatformCookedDir)"/>
 				<Delete Directories="$(PlatformCookedDir)"/>
 			</Do>
