Add a "Priming" null cook in the IterativeCompare job.

diff --git a/Engine/Build/Graph/Tests/IterativeCompare.xml b/Engine/Build/Graph/Tests/IterativeCompare.xml
--- a/Engine/Build/Graph/Tests/IterativeCompare.xml
+++ b/Engine/Build/Graph/Tests/IterativeCompare.xml
@@ -39,6 +39,7 @@
 	<Property Name="CompileNodeName" Value="Compile $(EditorPlatform) $(ProjectName)"/>
 	<Property Name="NaturalIncrementalCookNodeName" Value="Natural Incremental Cook $(ProjectName) $(TargetPlatform)"/>
 	<Property Name="StoreIncrementalCookNodeName" Value="Store Incremental Cook $(ProjectName) $(TargetPlatform)"/>
+	<Property Name="PrimingNullIncrementalCookNodeName" Value="Priming Null Incremental Cook $(ProjectName) $(TargetPlatform)"/>
 	<Property Name="NullIncrementalCookNodeName" Value="Null Incremental Cook $(ProjectName) $(TargetPlatform)"/>
 	<Property Name="CleanNullIncrementalCookNodeName" Value="Clean Null Incremental Cook $(ProjectName) $(TargetPlatform)"/>
 	<Property Name="FullCookNodeName" Value="Full Cook $(ProjectName) $(TargetPlatform)"/>
@@ -81,7 +82,14 @@
 		</Node>
 		<Property Name="FullCookPrequisiteNodeName" Value="$(StoreIncrementalCookNodeName)"/>
 		<Do If="$(WithNullCook)">
-			<Node Name="$(NullIncrementalCookNodeName)" Requires="$(StoreIncrementalCookNodeName)">
+			<Node Name="$(PrimingNullIncrementalCookNodeName)" Requires="$(StoreIncrementalCookNodeName)">
+				<!-- Execute the null incremental cook. -->
+				<Log Message="$(NewLine)========== IterativeCompare.xml: Priming Null Incremental Cook =========="/>
+				<Property Name="NullCookArgs" Value="$(CommonCookArgs)"/>
+				<Property Name="NullCookArgs" Value="$(NullCookArgs) -CustomConfig=$(NullCookCustomConfig)" If="'$(NullCookCustomConfig)' != ''"/>
+				<Commandlet Name="Cook" Project="$(ProjectFile)" Arguments="$(NullCookArgs)"/>
+			</Node>
+			<Node Name="$(NullIncrementalCookNodeName)" Requires="$(PrimingNullIncrementalCookNodeName)">
 				<!-- Execute the null incremental cook. -->
 				<Log Message="$(NewLine)========== IterativeCompare.xml: Null Incremental Cook =========="/>
 				<Property Name="NullCookArgs" Value="$(CommonCookArgs)"/>
