Allow null cook (and priming) to ignore blueprint cook dependencies.

diff --git a/Engine/Build/Graph/Tests/IterativeCompare.xml b/Engine/Build/Graph/Tests/IterativeCompare.xml
--- a/Engine/Build/Graph/Tests/IterativeCompare.xml
+++ b/Engine/Build/Graph/Tests/IterativeCompare.xml
@@ -85,14 +85,14 @@
 			<Node Name="$(PrimingNullIncrementalCookNodeName)" Requires="$(StoreIncrementalCookNodeName)">
 				<!-- Execute the null incremental cook. -->
 				<Log Message="$(NewLine)========== IterativeCompare.xml: Priming Null Incremental Cook =========="/>
-				<Property Name="NullCookArgs" Value="$(CommonCookArgs)"/>
+				<Property Name="NullCookArgs" Value="$(CommonCookArgs) -IgnoreBlueprintCookDependencies"/>
 				<Property Name="NullCookArgs" Value="$(NullCookArgs) -CustomConfig=$(NullCookCustomConfig)" If="'$(NullCookCustomConfig)' != ''"/>
 				<Commandlet Name="Cook" Project="$(ProjectFile)" Arguments="$(NullCookArgs)"/>
 			</Node>
 			<Node Name="$(NullIncrementalCookNodeName)" Requires="$(PrimingNullIncrementalCookNodeName)">
 				<!-- Execute the null incremental cook. -->
 				<Log Message="$(NewLine)========== IterativeCompare.xml: Null Incremental Cook =========="/>
-				<Property Name="NullCookArgs" Value="$(CommonCookArgs)"/>
+				<Property Name="NullCookArgs" Value="$(CommonCookArgs) -IgnoreBlueprintCookDependencies"/>
 				<Property Name="NullCookArgs" Value="$(NullCookArgs) -CustomConfig=$(NullCookCustomConfig)" If="'$(NullCookCustomConfig)' != ''"/>
 				<Commandlet Name="Cook" Project="$(ProjectFile)" Arguments="$(NullCookArgs)"/>
 			</Node>
