Add en-only cook culture argument and -unversioned to the IterativeCompare cook steps.

diff --git a/Engine/Build/Graph/Tests/IterativeCompare.xml b/Engine/Build/Graph/Tests/IterativeCompare.xml
--- a/Engine/Build/Graph/Tests/IterativeCompare.xml
+++ b/Engine/Build/Graph/Tests/IterativeCompare.xml
@@ -29,7 +29,7 @@
 	<Option Name="NullCookCustomConfig" DefaultValue="" Description="What custom configuration to apply if doing a null cook (see -CustomConfig commandline argument)"/>
 	<Property Name="FullPathProjectDir" Value="$(RootDir)/$(ProjectDir)"/>
 	<Property Name="EditorPlatform" Value="Win64"/>
-	<Property Name="CommonCookArgs" Value="-targetplatform=$(TargetPlatform) -zenstore -cookincremental"/>
+	<Property Name="CommonCookArgs" Value="-targetplatform=$(TargetPlatform) -cookcultures=en -unversioned -zenstore -cookincremental"/>
 	<Property Name="CommonCookArgs" Value="$(CommonCookArgs) -CookProcessCount=$(CookProcessCount)" If="'$(CookProcessCount)' != ''"/>
 	<Property Name="PlatformCookedDir" Value="$(FullPathProjectDir)/Saved/Cooked/$(TargetPlatform)"/>
 	<Property Name="CompareDirIncremental" Value="$(FullPathProjectDir)/Intermediate/IterativeCompare/Iter/$(TargetPlatform)"/>
