Add MPCookCoreSubscription parameter to IterativeCompare job to allow over-subscription of cores during cook.

diff --git a/Engine/Build/Graph/Tests/IterativeCompare.xml b/Engine/Build/Graph/Tests/IterativeCompare.xml
--- a/Engine/Build/Graph/Tests/IterativeCompare.xml
+++ b/Engine/Build/Graph/Tests/IterativeCompare.xml
@@ -25,12 +25,14 @@
 	
 	<!-- Agent=IterativeCompareAgent variables -->
 	<Option Name="CookProcessCount" DefaultValue="" Description="Number of processes to use during cooking, default is ini-driven."/>
+	<Option Name="MPCookCoreSubscription" DefaultValue="" Description="Over or Under subscription factor for CPU cores to use during cooking, default is ini-driven."/>
 	<Option Name="WithNullCook" Restrict="true|false" DefaultValue="false" Description="Whether to do a null cook or not"/>
 	<Option Name="NullCookCustomConfig" DefaultValue="" Description="What custom configuration to apply if doing a null cook (see -CustomConfig commandline argument)"/>
 	<Property Name="FullPathProjectDir" Value="$(RootDir)/$(ProjectDir)"/>
 	<Property Name="EditorPlatform" Value="Win64"/>
 	<Property Name="CommonCookArgs" Value="-targetplatform=$(TargetPlatform) -cookcultures=en -unversioned -zenstore -cookincremental"/>
 	<Property Name="CommonCookArgs" Value="$(CommonCookArgs) -CookProcessCount=$(CookProcessCount)" If="'$(CookProcessCount)' != ''"/>
+	<Property Name="CommonCookArgs" Value="$(CommonCookArgs) -MPCookCoreSubscription=$(MPCookCoreSubscription)" If="'$(MPCookCoreSubscription)' != ''"/>
 	<Property Name="PlatformCookedDir" Value="$(FullPathProjectDir)/Saved/Cooked/$(TargetPlatform)"/>
 	<Property Name="CompareDirIncremental" Value="$(FullPathProjectDir)/Intermediate/IterativeCompare/Iter/$(TargetPlatform)"/>
 	<Property Name="CompareDirRecook" Value="$(FullPathProjectDir)/Intermediate/IterativeCompare/Recook/$(TargetPlatform)"/>
