Double-publish zen snapshots to both generic cloud and builds API from BuildAndTest graph.

diff --git a/Engine/Build/Graph/Tasks/BuildAndTestProject.xml b/Engine/Build/Graph/Tasks/BuildAndTestProject.xml
--- a/Engine/Build/Graph/Tasks/BuildAndTestProject.xml
+++ b/Engine/Build/Graph/Tasks/BuildAndTestProject.xml
@@ -10,6 +10,7 @@
 	<EnvVar Name="UE-CloudPublishDescriptorHost"/>
 	<EnvVar Name="UE-CloudPublishDescriptorHttpVersion"/>
 	<EnvVar Name="UE-CloudPublishNamespace"/>
+	<EnvVar Name="UE-CloudPublishBaselineBranch"/>
 	<EnvVar Name="UE-ZenPublishHost"/>
 
 	<Property Name="IsHordeEnv" Value="false"/>
@@ -660,7 +661,7 @@
 						</ForEach>
 					</Case>
 					<Default>
-						<Command Name="BuildCookRun" Arguments="$(BCRArgs) -configuration=$(PackageConfigurations) $(StageAndPackageArguments) -skipbuild -skipcook -stage -pak" />
+						<Command Name="BuildCookRun" Arguments="$(BCRArgs) -configuration=$(PackageConfigurations) $(StageAndPackageArguments) -skipbuild -skipcook -stage -pak -NoCodeSign" />
 					</Default>
 				</Switch>
 			</Node>
@@ -693,13 +694,26 @@
 					<Property Name="SnapshotLongTermStorageDir" Value="$(NetworkSnapshotDescFileRoot)/$(EscapedBranch)/$(ProjectName)/$(SnapshotLongTermPlatformFolder)/$(Change)"  If="'$(NetworkSnapshotDescFileRoot)' != ''"/>
 
 					<Do If="$(ShouldSnapshotPublishCloud)">
+						<!-- Publish to generic cloud storage endpoint -->
 						<Property Name="SnapshotBaseDescriptorFile" Value=""/>
 						<Property Name="SnapshotBaseDescriptorFile" Value="$(NetworkSnapshotDescFileRoot)/$(EscapedBranch)/$(ProjectName)/$(StagedPlatformFolder)/$(UE_HORDE_LAST_WARNING_CL)/$(SnapshotFilenamePrefix)-cloud.json"  If="'$(NetworkSnapshotDescFileRoot)' != '' and '$(UE_HORDE_LAST_WARNING_CL)' != ''"/>
-						<ZenExportSnapshot Project="$(TargetProject)" Platform="$(CookPlatform)" DestinationStorageType="Cloud" DestinationCloudHost="$(UE-CloudPublishHost)" DestinationCloudHttpVersion="$(UE-CloudPublishHttpVersion)" SnapshotDescriptorCloudHost="$(UE-CloudPublishDescriptorHost)" SnapshotDescriptorCloudHttpVersion="$(UE-CloudPublishDescriptorHttpVersion)" DestinationCloudNamespace="$(UE-CloudPublishNamespace)" DestinationCloudBucket="$(EscapedBranch).$(CookPlatform)" DestinationIdentifier="$(SnapshotIdentifier)" SnapshotDescriptorFile="$(SnapshotLocalDir)/$(StagedPlatformFolder)/$(SnapshotFilenamePrefix)-cloud.json" SnapshotBaseDescriptorFile="$(SnapshotBaseDescriptorFile)"/>
+						<ZenExportSnapshot Project="$(TargetProject)" Platform="$(CookPlatform)" DestinationStorageType="Cloud" DestinationCloudHost="$(UE-CloudPublishHost)" DestinationCloudHttpVersion="$(UE-CloudPublishHttpVersion)" SnapshotDescriptorCloudHost="$(UE-CloudPublishDescriptorHost)" SnapshotDescriptorCloudHttpVersion="$(UE-CloudPublishDescriptorHttpVersion)" DestinationCloudNamespace="$(UE-CloudPublishNamespace)" DestinationCloudBucket="$(ProjectName).oplog.$(EscapedBranch).$(CookPlatform)" DestinationIdentifier="$(SnapshotIdentifier)" SnapshotDescriptorFile="$(SnapshotLocalDir)/$(StagedPlatformFolder)/$(SnapshotFilenamePrefix)-cloud.json" SnapshotBaseDescriptorFile="$(SnapshotBaseDescriptorFile)"/>
 						<Copy From="$(SnapshotLocalDir)/$(StagedPlatformFolder)/$(SnapshotFilenamePrefix)-cloud.json" To="$(NetworkOutputDirectory)/$(StagedPlatformFolder)/SnapshotDescriptor.json"/>
 						<Do If="'$(SnapshotLongTermStorageDir)' != ''">
 							<Copy From="$(SnapshotLocalDir)/$(StagedPlatformFolder)/$(SnapshotFilenamePrefix)-cloud.json" To="$(SnapshotLongTermStorageDir)/$(SnapshotFilenamePrefix)-cloud.json"/>
 						</Do>
+						<!-- Publish to builds cloud storage endpoint -->
+						<Property Name="SnapshotBuildBaselineBranch" Value="$(Branch)"/>
+						<Property Name="SnapshotBuildBaselineBranch" Value="$(UE-CloudPublishBaselineBranch)" If="'$(UE-CloudPublishBaselineBranch)' != ''"/>
+						<Property Name="SnapshotBuildMetaData" Value="name=$(SnapshotIdentifier);depot=$(Depot);branch=$(Branch);baselineBranch=$(SnapshotBuildBaselineBranch);project=$(ProjectName);runtime=$(TargetType);platform=$(TargetPlatform);cookPlatform=$(CookPlatform);changelist=$(Change);codeChangelist=$(CodeChange);isPreflight=$(IsPreflight);preflightChangelist=$(PreflightChange);isBuildMachine=$(IsBuildMachine);isHordeEnv=$(IsHordeEnv)"/>
+						<Property Name="SnapshotBuildMetaData" Value="$(SnapshotBuildMetaData);cookFlavor=$(BCRCookFlavor)" If="'$(BCRCookFlavor)' != ''" />
+						<Property Name="SnapshotBuildMetaData" Value="$(SnapshotBuildMetaData);hordeTemplateId=$(UE_HORDE_TEMPLATEID)" If="'$(UE_HORDE_TEMPLATEID)' != ''" />
+						<Log Message="Publishing build with metadata '$(SnapshotBuildMetaData)'"/>
+						<ZenExportSnapshot Project="$(TargetProject)" Platform="$(CookPlatform)" DestinationStorageType="Builds" DestinationCloudHost="$(UE-CloudPublishHost)" DestinationCloudHttpVersion="$(UE-CloudPublishHttpVersion)" SnapshotDescriptorCloudHost="$(UE-CloudPublishDescriptorHost)" SnapshotDescriptorCloudHttpVersion="$(UE-CloudPublishDescriptorHttpVersion)" DestinationCloudNamespace="$(UE-CloudPublishNamespace)" DestinationCloudBucket="$(ProjectName).oplog.$(EscapedBranch).$(CookPlatform)" DestinationIdentifier="$(SnapshotIdentifier)" SnapshotDescriptorFile="$(SnapshotLocalDir)/$(StagedPlatformFolder)/$(SnapshotFilenamePrefix)-builds.json" Metadata="$(SnapshotBuildMetaData)"/>
+						<Copy From="$(SnapshotLocalDir)/$(StagedPlatformFolder)/$(SnapshotFilenamePrefix)-builds.json" To="$(NetworkOutputDirectory)/$(StagedPlatformFolder)/SnapshotDescriptorBuilds.json"/>
+						<Do If="'$(SnapshotLongTermStorageDir)' != ''">
+							<Copy From="$(SnapshotLocalDir)/$(StagedPlatformFolder)/$(SnapshotFilenamePrefix)-builds.json" To="$(SnapshotLongTermStorageDir)/$(SnapshotFilenamePrefix)-builds.json"/>
+						</Do>
 					</Do>
 					<Do If="$(ShouldSnapshotPublishZen)">
 						<ZenExportSnapshot Project="$(TargetProject)" Platform="$(CookPlatform)" DestinationStorageType="Zen" DestinationZenHost="$(UE-ZenPublishHost)" DestinationIdentifier="$(SnapshotIdentifier)" SnapshotDescriptorFile="$(SnapshotLocalDir)/$(StagedPlatformFolder)/$(SnapshotFilenamePrefix)-zen.json"/>
