Avoid specifying preflightChangelist metadata if it is empty.

diff --git a/Engine/Build/Graph/Tasks/BuildAndTestProject.xml b/Engine/Build/Graph/Tasks/BuildAndTestProject.xml
--- a/Engine/Build/Graph/Tasks/BuildAndTestProject.xml
+++ b/Engine/Build/Graph/Tasks/BuildAndTestProject.xml
@@ -705,8 +705,9 @@
 						<!-- Publish to builds cloud storage endpoint -->
 						<Property Name="SnapshotBuildBaselineBranch" Value="$(Branch)"/>
 						<Property Name="SnapshotBuildBaselineBranch" Value="$(UE-CloudPublishBaselineBranch)" If="'$(UE-CloudPublishBaselineBranch)' != ''"/>
-						<Property Name="SnapshotBuildMetaData" Value="name=$(SnapshotIdentifier);depot=$(Depot);branch=$(Branch);baselineBranch=$(SnapshotBuildBaselineBranch);project=$(ProjectName);runtime=$(TargetType);platform=$(TargetPlatform);cookPlatform=$(CookPlatform);changelist=$(Change);codeChangelist=$(CodeChange);isPreflight=$(IsPreflight);preflightChangelist=$(PreflightChange);isBuildMachine=$(IsBuildMachine);isHordeEnv=$(IsHordeEnv)"/>
+						<Property Name="SnapshotBuildMetaData" Value="name=$(SnapshotIdentifier);depot=$(Depot);branch=$(Branch);baselineBranch=$(SnapshotBuildBaselineBranch);project=$(ProjectName);runtime=$(TargetType);platform=$(TargetPlatform);cookPlatform=$(CookPlatform);changelist=$(Change);codeChangelist=$(CodeChange);isPreflight=$(IsPreflight);isBuildMachine=$(IsBuildMachine);isHordeEnv=$(IsHordeEnv)"/>
 						<Property Name="SnapshotBuildMetaData" Value="$(SnapshotBuildMetaData);cookFlavor=$(BCRCookFlavor)" If="'$(BCRCookFlavor)' != ''" />
+						<Property Name="SnapshotBuildMetaData" Value="$(SnapshotBuildMetaData);preflightChangelist=$(PreflightChange)" If="'$(PreflightChange)' != ''" />
 						<Property Name="SnapshotBuildMetaData" Value="$(SnapshotBuildMetaData);hordeTemplateId=$(UE_HORDE_TEMPLATEID)" If="'$(UE_HORDE_TEMPLATEID)' != ''" />
 						<Log Message="Publishing build with metadata '$(SnapshotBuildMetaData)'"/>
 						<ZenExportSnapshot Project="$(TargetProject)" Platform="$(CookPlatform)" DestinationStorageType="Builds" DestinationCloudHost="$(UE-CloudPublishHost)" DestinationCloudHttpVersion="$(UE-CloudPublishHttpVersion)" SnapshotDescriptorCloudHost="$(UE-CloudPublishDescriptorHost)" SnapshotDescriptorCloudHttpVersion="$(UE-CloudPublishDescriptorHttpVersion)" DestinationCloudNamespace="$(UE-CloudPublishNamespace)" DestinationCloudBucket="$(ProjectName).oplog.$(EscapedBranch).$(CookPlatform)" DestinationIdentifier="$(SnapshotIdentifier)" SnapshotDescriptorFile="$(SnapshotLocalDir)/$(StagedPlatformFolder)/$(SnapshotFilenamePrefix)-builds.json" Metadata="$(SnapshotBuildMetaData)"/>
