Change IterativeCompare null cook step to use display mode 2 so we get a listing of packages being cooked.

diff --git a/Engine/Build/Graph/Tests/IterativeCompare.xml b/Engine/Build/Graph/Tests/IterativeCompare.xml
--- a/Engine/Build/Graph/Tests/IterativeCompare.xml
+++ b/Engine/Build/Graph/Tests/IterativeCompare.xml
@@ -94,7 +94,7 @@
 			<Node Name="$(NullIncrementalCookNodeName)" Requires="$(PrimingNullIncrementalCookNodeName)">
 				<!-- Execute the null incremental cook. -->
 				<Log Message="$(NewLine)========== IterativeCompare.xml: Null Incremental Cook =========="/>
-				<Property Name="NullCookArgs" Value="$(CommonCookArgs) -IgnoreBlueprintCookDependencies"/>
+				<Property Name="NullCookArgs" Value="$(CommonCookArgs) -IgnoreBlueprintCookDependencies -dpcvars=cook.displaymode=2"/>
 				<Property Name="NullCookArgs" Value="$(NullCookArgs) -CustomConfig=$(NullCookCustomConfig)" If="'$(NullCookCustomConfig)' != ''"/>
 				<Commandlet Name="Cook" Project="$(ProjectFile)" Arguments="$(NullCookArgs)"/>
 			</Node>
