Testing Jupiter Artifact storage

#rnx

diff --git a/Engine/Build/Graph/Tasks/BuildAndTestProject.xml b/Engine/Build/Graph/Tasks/BuildAndTestProject.xml
--- a/Engine/Build/Graph/Tasks/BuildAndTestProject.xml
+++ b/Engine/Build/Graph/Tasks/BuildAndTestProject.xml
@@ -17,6 +17,8 @@
 	<Property Name="IsHordeEnv" Value="true" If="'$(UE_HORDE_JOBID)' != ''"/>
 	<Option Name="Horde" DefaultValue="$(IsBuildMachine) And $(IsHordeEnv)" Description="Whether the build is running under Horde"/>
 	
+	<Option Name="PublishToCloudDDC" DefaultValue="false" Description="Enable to publish to cloud ddc instead of fileshare" />
+
 	<!-- Project Options -->
 	<Option Name="ProjectFile" DefaultValue="" Restrict=".*" Description="Path to the project file. Can be used instead of specifying ProjectName and ProjectPath separately."/>
 
@@ -570,6 +572,8 @@
 		<Property Name="PlatformStagedArtifactNodeName" Value="$(PlatformPublishStagedNodeName) (Horde)" />
 		<Property Name="HordeArtifactPlatformPackagedNodeName" Value="$(PlatformPublishPackagedNodeName) (Horde)" />
 		<Property Name="HordeArtifactPlatformRetrieve" Value="$(PreNodeNameWithPlatform)Retrieve Build Artifact $(TargetPlatform)" />
+		<Property Name="CloudDDCArtifactPlatformPackagedNodeName" Value="$(PlatformPublishPackagedNodeName) (CloudDDC)" />
+		<Property Name="CloudDDCArtifactPlatformRetrieve" Value="$(PreNodeNameWithPlatform)Retrieve Build Artifact $(TargetPlatform)" />
 		<Property Name="RequirePlatformStaged" Value="$(PlatformStageNodeName)" />
 		<Property Name="RequirePlatformStaged" Value="$(RequirePlatformStaged);$(PlatformStagedArtifactNodeName)" If="$(Horde)"/>
 		<Property Name="RequirePlatformStaged" Value="$(RequirePlatformStaged);$(PlatformPublishStagedNodeName)" If="!$(SkipPublish) And $(PublishBuilds)"/>
@@ -778,6 +782,19 @@
 				</Node>
 			</Do>
 
+			<!-- Publish (Packages) (CloudDDC) -->
+			<Property Name="CloudPackageSkipBCRArguments" Value="-stage -package"/>
+			<Property Name="CloudPackageSkipBCRArguments" Value="-skipstage -skippackage" If="'$(TargetPlatform)' == 'Mac'"/>
+
+			<Node Name="$(CloudDDCArtifactPlatformPackagedNodeName)" Requires="$(PlatformPackageNodeName)">
+				<ForEach Name="PackageConfiguration" Values="$(PackageConfigurations)" Separator="+">
+					<Command Name="BuildCookRun" Arguments="$(BCRArgs) -configuration=$(PackageConfiguration) $(StageAndPackageArguments) -skipbuild -skipcook -skippak $(CloudPackageSkipBCRArguments) -archive -archivedirectory=&quot;$(ProjectOutputDirectory)/Packages/$(StagedPlatformFolder)/$(PackageConfiguration)&quot;" />
+					<Tag Files="$(ProjectOutputDirectory)/Packages/$(StagedPlatformFolder)/$(PackageConfiguration)/..." With="#$(PlatformPackageNodeName) $(PackageConfiguration) Files"/>
+
+					<CreateCloudArtifact Name="$(ProjectName)-Packaged-$(StagedPlatformFolder)-$(PackageConfiguration)" Type="packaged-build" BaseDir="$(ProjectOutputDirectory)/Packages/$(StagedPlatformFolder)/$(PackageConfiguration)" Metadata="Project=$(ProjectName);Stream=$(Branch);Platform=$(StagedPlatformFolder);Configuration=$(PackageConfiguration);BuildName=$(BuildNamePath);dashboard-category=$(PackageConfiguration) Builds" Project="$(ProjectName)" Platform="$(StagedPlatformFolder)" BaselineBranch="$(UE-CloudPublishBaselineBranch)" Files="#$(PlatformPackageNodeName) $(PackageConfiguration) Files"/>
+				</ForEach>
+			</Node>
+
 			<!-- Check Build Size -->
 			<Node Name="$(PlatformCheckBuildSizeNodeName)" Requires="$(PlatformPackageNodeName)">
 				<ForEach Name="PackageConfiguration" Values="$(PackageConfigurations)" Separator="+">
