Fixing new Juptier upload step to use "do if"

#rnx

diff --git a/Engine/Build/Graph/Tasks/BuildAndTestProject.xml b/Engine/Build/Graph/Tasks/BuildAndTestProject.xml
--- a/Engine/Build/Graph/Tasks/BuildAndTestProject.xml
+++ b/Engine/Build/Graph/Tasks/BuildAndTestProject.xml
@@ -783,17 +783,19 @@
 			</Do>
 
 			<!-- Publish (Packages) (CloudDDC) -->
-			<Property Name="CloudPackageSkipBCRArguments" Value="-stage -package"/>
-			<Property Name="CloudPackageSkipBCRArguments" Value="-skipstage -skippackage" If="'$(TargetPlatform)' == 'Mac'"/>
+			<Do  If="$(PublishToCloudDDC)">
+				<Property Name="CloudPackageSkipBCRArguments" Value="-stage -package"/>
+				<Property Name="CloudPackageSkipBCRArguments" Value="-skipstage -skippackage" If="'$(TargetPlatform)' == 'Mac'"/>
 
-			<Node Name="$(CloudDDCArtifactPlatformPackagedNodeName)" Requires="$(PlatformPackageNodeName)">
-				<ForEach Name="PackageConfiguration" Values="$(PackageConfigurations)" Separator="+">
-					<Command Name="BuildCookRun" Arguments="$(BCRArgs) -configuration=$(PackageConfiguration) $(StageAndPackageArguments) -skipbuild -skipcook -skippak $(CloudPackageSkipBCRArguments) -archive -archivedirectory=&quot;$(ProjectOutputDirectory)/Packages/$(StagedPlatformFolder)/$(PackageConfiguration)&quot;" />
-					<Tag Files="$(ProjectOutputDirectory)/Packages/$(StagedPlatformFolder)/$(PackageConfiguration)/..." With="#$(PlatformPackageNodeName) $(PackageConfiguration) Files"/>
+				<Node Name="$(CloudDDCArtifactPlatformPackagedNodeName)" Requires="$(PlatformPackageNodeName)">
+					<ForEach Name="PackageConfiguration" Values="$(PackageConfigurations)" Separator="+">
+						<Command Name="BuildCookRun" Arguments="$(BCRArgs) -configuration=$(PackageConfiguration) $(StageAndPackageArguments) -skipbuild -skipcook -skippak $(CloudPackageSkipBCRArguments) -archive -archivedirectory=&quot;$(ProjectOutputDirectory)/Packages/$(StagedPlatformFolder)/$(PackageConfiguration)&quot;" />
+						<Tag Files="$(ProjectOutputDirectory)/Packages/$(StagedPlatformFolder)/$(PackageConfiguration)/..." With="#$(PlatformPackageNodeName) $(PackageConfiguration) Files"/>
 
-					<CreateCloudArtifact Name="$(ProjectName)-Packaged-$(StagedPlatformFolder)-$(PackageConfiguration)" Type="packaged-build" BaseDir="$(ProjectOutputDirectory)/Packages/$(StagedPlatformFolder)/$(PackageConfiguration)" Metadata="Project=$(ProjectName);Stream=$(Branch);Platform=$(StagedPlatformFolder);Configuration=$(PackageConfiguration);BuildName=$(BuildNamePath);dashboard-category=$(PackageConfiguration) Builds" Project="$(ProjectName)" Platform="$(StagedPlatformFolder)" BaselineBranch="$(UE-CloudPublishBaselineBranch)" Files="#$(PlatformPackageNodeName) $(PackageConfiguration) Files"/>
-				</ForEach>
-			</Node>
+						<CreateCloudArtifact Name="$(ProjectName)-Packaged-$(StagedPlatformFolder)-$(PackageConfiguration)" Type="packaged-build" BaseDir="$(ProjectOutputDirectory)/Packages/$(StagedPlatformFolder)/$(PackageConfiguration)" Metadata="Project=$(ProjectName);Stream=$(Branch);Platform=$(StagedPlatformFolder);Configuration=$(PackageConfiguration);BuildName=$(BuildNamePath);dashboard-category=$(PackageConfiguration) Builds" Project="$(ProjectName)" Platform="$(StagedPlatformFolder)" BaselineBranch="$(UE-CloudPublishBaselineBranch)" Files="#$(PlatformPackageNodeName) $(PackageConfiguration) Files"/>
+					</ForEach>
+				</Node>
+			</Do>>
 
 			<!-- Check Build Size -->
 			<Node Name="$(PlatformCheckBuildSizeNodeName)" Requires="$(PlatformPackageNodeName)">
