fixing stray >

#rnx

diff --git a/Engine/Build/Graph/Tasks/BuildAndTestProject.xml b/Engine/Build/Graph/Tasks/BuildAndTestProject.xml
--- a/Engine/Build/Graph/Tasks/BuildAndTestProject.xml
+++ b/Engine/Build/Graph/Tasks/BuildAndTestProject.xml
@@ -795,7 +795,7 @@
 						<CreateCloudArtifact Name="$(ProjectName)-Packaged-$(StagedPlatformFolder)-$(PackageConfiguration)" Type="packaged-build" BaseDir="$(ProjectOutputDirectory)/Packages/$(StagedPlatformFolder)/$(PackageConfiguration)" Metadata="Project=$(ProjectName);Stream=$(Branch);Platform=$(StagedPlatformFolder);Configuration=$(PackageConfiguration);BuildName=$(BuildNamePath);dashboard-category=$(PackageConfiguration) Builds" Project="$(ProjectName)" Platform="$(StagedPlatformFolder)" BaselineBranch="$(UE-CloudPublishBaselineBranch)" Files="#$(PlatformPackageNodeName) $(PackageConfiguration) Files"/>
 					</ForEach>
 				</Node>
-			</Do>>
+			</Do>
 
 			<!-- Check Build Size -->
 			<Node Name="$(PlatformCheckBuildSizeNodeName)" Requires="$(PlatformPackageNodeName)">
