fixing jupiter testing

#rnx

diff --git a/Engine/Build/Graph/Tasks/BuildAndTestProject.xml b/Engine/Build/Graph/Tasks/BuildAndTestProject.xml
--- a/Engine/Build/Graph/Tasks/BuildAndTestProject.xml
+++ b/Engine/Build/Graph/Tasks/BuildAndTestProject.xml
@@ -579,6 +579,7 @@
 		<Property Name="RequirePlatformStaged" Value="$(RequirePlatformStaged);$(PlatformPublishStagedNodeName)" If="!$(SkipPublish) And $(PublishBuilds)"/>
 		<Property Name="RequirePlatformPackaged" Value="$(PlatformPackageNodeName)" />
 		<Property Name="RequirePlatformPackaged" Value="$(RequirePlatformPackaged);$(HordeArtifactPlatformPackagedNodeName)" If="$(Horde)"/>
+		<Property Name="RequirePlatformPackaged" Value="$(RequirePlatformPackaged);$(CloudDDCArtifactPlatformPackagedNodeName)" If="$(PublishToCloudDDC)"/>
 		<Property Name="RequirePlatformPackaged" Value="$(RequirePlatformPackaged);$(PlatformPublishPackagedNodeName)"  If="!$(SkipPublish) And $(PublishBuilds)"/>
 		
 		<!-- Set CookPlatform since it isn't always the same as the target -->
