Allow IterativeValidate job to optionally capture utrace files.

diff --git a/Engine/Build/Graph/Tests/IterativeValidate.xml b/Engine/Build/Graph/Tests/IterativeValidate.xml
--- a/Engine/Build/Graph/Tests/IterativeValidate.xml
+++ b/Engine/Build/Graph/Tests/IterativeValidate.xml
@@ -27,11 +27,19 @@
 	<Property Name="IterativeArgs" Value="-zenstore -cookincremental" If="'$(IterativeMode)' == 'incremental'"/>
 	<Property Name="IterativeArgs" Value="-iterate -ini:Editor:[CookSettings]:HybridIterativeEnabled=false -ini:Editor:[CookSettings]:IterativeUseClassFilters=true" If="'$(IterativeMode)' == 'iterative'"/>
 
+	<!-- Profiling Options -->
+	<Option Name="ProfilingFilesPublishRootPath" DefaultValue="" Description="Where to publish profiling data files."/>
+	<Option Name="WithCookUTRACECapture" DefaultValue="false" Restrict="true|false" Description="Whether to capture UTRACE data during cooks. Only true or false"/>
+
 	<!-- Config Options -->
 	<Option Name="CustomConfig" DefaultValue="" Description="What custom configuration to apply (see -CustomConfig commandline argument)"/>
 
 	<!-- Agent=IterativeValidate variables -->
 	<Property Name="EditorPlatform" Value="Win64"/>
+	<Property Name="TraceArgs" Value=""/>
+	<Property Name="TraceArgs" Value="-trace=cpu,loadtime,savetime,cook,contextswitch" If="$(WithCookUTRACECapture) and '$(ProfilingFilesPublishRootPath)' != ''"/>
+	<Property Name="ProfilingFilesLocalRootPath" Value="$(RootDir)/Engine/Intermediate/Profiling"/>
+
 	<Property Name="CommonCookArgs" Value="-targetplatform=$(TargetPlatform) $(IterativeArgs)"/>
 	<Property Name="CommonCookArgs" Value="$(CommonCookArgs) -CustomConfig=$(CustomConfig)" If="'$(CustomConfig)' != ''"/>
 
@@ -46,13 +54,28 @@
 			<Compile Target="$(EditorTarget)" Platform="$(EditorPlatform)" Project="$(ProjectFile)" Configuration="Development" Arguments="$(UBACompileArgs)"/>
 		</Node>
 		<Node Name="$(CookIterativeValidateSinglePhaseName)" Requires="$(CompileNodeName)">
-			<Commandlet Name="Cook" Project="$(ProjectFile)" Arguments="$(CommonCookArgs) -IterativeValidate -IterativeValidateAllowWrite"/>
+			<Property Name="TraceArgsFinal" Value=""/>
+			<Property Name="TraceArgsFinal" Value="$(TraceArgs) -tracefile=&quot;$(ProfilingFilesLocalRootPath)/$(ProjectName)_IterativeValidate.utrace&quot;" If="$(WithCookUTRACECapture) and '$(ProfilingFilesPublishRootPath)' != ''"/>
+			<Commandlet Name="Cook" Project="$(ProjectFile)" Arguments="$(CommonCookArgs) $(TraceArgsFinal) -IterativeValidate -IterativeValidateAllowWrite"/>
+			<Do If="'$(TraceArgsFinal)' != ''">
+				<Tag Files="$($(Project)ProjectPath)/Intermediate/Build/$(Platform)$(Subfolder)$(Project)/..." Filter="*.csv" With="#$(Platform)$(Project)TimingFiles"/>
+				<Zip FromDir="$(ProfilingFilesLocalRootPath)" Files="*.utrace" ZipFile="$(ProfilingFilesLocalRootPath)/$(ProjectName)_IterativeValidate.zip"/>
+				<Copy Files="$(ProjectName)_IterativeValidate.zip" From="$(ProfilingFilesLocalRootPath)" To="$(ProfilingFilesPublishRootPath)\$(EscapedBranch)-CL-$(Change)"/>
+			</Do>
 		</Node>
 		<Node Name="$(CookIterativeValidatePhase1Name)" Requires="$(CompileNodeName)">
-			<Commandlet Name="Cook" Project="$(ProjectFile)" Arguments="$(CommonCookArgs) -IterativeValidatePhase1 -IterativeValidateAllowWrite"/>
+			<Property Name="TraceArgsFinal" Value=""/>
+			<Property Name="TraceArgsFinal" Value="$(TraceArgs) -tracefile=&quot;$(ProfilingFilesLocalRootPath)/$(ProjectName)_IterativeValidatePhase1.utrace&quot;" If="$(WithCookUTRACECapture) and '$(ProfilingFilesPublishRootPath)' != ''"/>
+			<Commandlet Name="Cook" Project="$(ProjectFile)" Arguments="$(CommonCookArgs) $(TraceArgsFinal) -IterativeValidatePhase1 -IterativeValidateAllowWrite"/>
 		</Node>
 		<Node Name="$(CookIterativeValidatePhase2Name)" Requires="$(CookIterativeValidatePhase1Name)">
-			<Commandlet Name="Cook" Project="$(ProjectFile)" Arguments="$(CommonCookArgs) -IterativeValidatePhase2 -IterativeValidateAllowWrite"/>
+			<Property Name="TraceArgsFinal" Value=""/>
+			<Property Name="TraceArgsFinal" Value="$(TraceArgs) -tracefile=&quot;$(ProfilingFilesLocalRootPath)/$(ProjectName)_IterativeValidatePhase2.utrace&quot;" If="$(WithCookUTRACECapture) and '$(ProfilingFilesPublishRootPath)' != ''"/>
+			<Commandlet Name="Cook" Project="$(ProjectFile)" Arguments="$(CommonCookArgs) $(TraceArgsFinal) -IterativeValidatePhase2 -IterativeValidateAllowWrite"/>
+			<Do If="'$(TraceArgsFinal)' != ''">
+				<Zip FromDir="$(ProfilingFilesLocalRootPath)" Files="*.utrace" ZipFile="$(ProfilingFilesLocalRootPath)/$(ProjectName)_IterativeValidate.zip"/>
+				<Copy Files="$(ProjectName)_IterativeValidate.zip" From="$(ProfilingFilesLocalRootPath)" To="$(ProfilingFilesPublishRootPath)\$(EscapedBranch)-CL-$(Change)"/>
+			</Do>
 		</Node>
 	</Agent>
 
