Remove superfluous line added to IterativeValidate build graph.

diff --git a/Engine/Build/Graph/Tests/IterativeValidate.xml b/Engine/Build/Graph/Tests/IterativeValidate.xml
--- a/Engine/Build/Graph/Tests/IterativeValidate.xml
+++ b/Engine/Build/Graph/Tests/IterativeValidate.xml
@@ -58,7 +58,6 @@
 			<Property Name="TraceArgsFinal" Value="$(TraceArgs) -tracefile=&quot;$(ProfilingFilesLocalRootPath)/$(ProjectName)_IterativeValidate.utrace&quot;" If="$(WithCookUTRACECapture) and '$(ProfilingFilesPublishRootPath)' != ''"/>
 			<Commandlet Name="Cook" Project="$(ProjectFile)" Arguments="$(CommonCookArgs) $(TraceArgsFinal) -IterativeValidate -IterativeValidateAllowWrite"/>
 			<Do If="'$(TraceArgsFinal)' != ''">
-				<Tag Files="$($(Project)ProjectPath)/Intermediate/Build/$(Platform)$(Subfolder)$(Project)/..." Filter="*.csv" With="#$(Platform)$(Project)TimingFiles"/>
 				<Zip FromDir="$(ProfilingFilesLocalRootPath)" Files="*.utrace" ZipFile="$(ProfilingFilesLocalRootPath)/$(ProjectName)_IterativeValidate.zip"/>
 				<Copy Files="$(ProjectName)_IterativeValidate.zip" From="$(ProfilingFilesLocalRootPath)" To="$(ProfilingFilesPublishRootPath)\$(EscapedBranch)-CL-$(Change)"/>
 			</Do>
