Ensure we delete trace files gathered from iterative validate before we gather new ones.

diff --git a/Engine/Build/Graph/Tests/IterativeValidate.xml b/Engine/Build/Graph/Tests/IterativeValidate.xml
--- a/Engine/Build/Graph/Tests/IterativeValidate.xml
+++ b/Engine/Build/Graph/Tests/IterativeValidate.xml
@@ -56,6 +56,9 @@
 		<Node Name="$(CookIterativeValidateSinglePhaseName)" Requires="$(CompileNodeName)">
 			<Property Name="TraceArgsFinal" Value=""/>
 			<Property Name="TraceArgsFinal" Value="$(TraceArgs) -tracefile=&quot;$(ProfilingFilesLocalRootPath)/$(ProjectName)_IterativeValidate.utrace&quot;" If="$(WithCookUTRACECapture) and '$(ProfilingFilesPublishRootPath)' != ''"/>
+			<Do If="'$(TraceArgsFinal)' != '' and Exists('$(ProfilingFilesLocalRootPath)')">
+				<Delete Files="$(ProfilingFilesLocalRootPath)/..."/>
+			</Do>
 			<Commandlet Name="Cook" Project="$(ProjectFile)" Arguments="$(CommonCookArgs) $(TraceArgsFinal) -IterativeValidate -IterativeValidateAllowWrite"/>
 			<Do If="'$(TraceArgsFinal)' != ''">
 				<Zip FromDir="$(ProfilingFilesLocalRootPath)" Files="*.utrace" ZipFile="$(ProfilingFilesLocalRootPath)/$(ProjectName)_IterativeValidate.zip"/>
@@ -65,6 +68,9 @@
 		<Node Name="$(CookIterativeValidatePhase1Name)" Requires="$(CompileNodeName)">
 			<Property Name="TraceArgsFinal" Value=""/>
 			<Property Name="TraceArgsFinal" Value="$(TraceArgs) -tracefile=&quot;$(ProfilingFilesLocalRootPath)/$(ProjectName)_IterativeValidatePhase1.utrace&quot;" If="$(WithCookUTRACECapture) and '$(ProfilingFilesPublishRootPath)' != ''"/>
+			<Do If="'$(TraceArgsFinal)' != '' and Exists('$(ProfilingFilesLocalRootPath)')">
+				<Delete Files="$(ProfilingFilesLocalRootPath)/..."/>
+			</Do>
 			<Commandlet Name="Cook" Project="$(ProjectFile)" Arguments="$(CommonCookArgs) $(TraceArgsFinal) -IterativeValidatePhase1 -IterativeValidateAllowWrite"/>
 		</Node>
 		<Node Name="$(CookIterativeValidatePhase2Name)" Requires="$(CookIterativeValidatePhase1Name)">
