Extend the BuildAndTestProject.xml BuildGraph script to support cook flavors and update RenderingValidationTests_Quest3.xml to specify the XR cook flavor.

diff --git a/Engine/Build/Graph/Tasks/BuildAndTestProject.xml b/Engine/Build/Graph/Tasks/BuildAndTestProject.xml
--- a/Engine/Build/Graph/Tasks/BuildAndTestProject.xml
+++ b/Engine/Build/Graph/Tasks/BuildAndTestProject.xml
@@ -588,8 +588,9 @@
 		<Property Name="CookPlatform" Value="$(CookPlatform)_ETC2" If="'$(TargetPlatform)' == 'Android'"/>
 		<Property Name="CookPlatform" Value="$(CookPlatform)Client" If="'$(TargetType)' == 'Client'" />
 
-		<Property Name="BCRCookFlavor" Value="" If="" />
-		<Property Name="BCRCookFlavor" Value="ETC2" If="'$(TargetPlatform)' == 'Android'" />
+		<!-- Only use Cook Flavor for Android -->
+		<Property Name="BCRCookFlavor" Value=""/>
+        <Property Name="BCRCookFlavor" Value="$(AndroidCookFlavor)" If="'$(TargetPlatform)' == 'Android'" />
 		
 		<Property Name="StagedPlatformFolder" Value="$(CookPlatform)"  />
 		<!-- <Property Name="StagedPlatformFolder" Value="$(StagedPlatformFolder)Client" If="'$(TargetType)' == 'Client'" /> -->
diff --git a/Engine/Build/Graph/Tasks/Inc/PlatformOptions.xml b/Engine/Build/Graph/Tasks/Inc/PlatformOptions.xml
--- a/Engine/Build/Graph/Tasks/Inc/PlatformOptions.xml
+++ b/Engine/Build/Graph/Tasks/Inc/PlatformOptions.xml
@@ -84,6 +84,8 @@
 <Option Name="WithTVOS" DefaultValue="false" Description="Build and test TVOS"/>
 <Option Name="WithWin64" DefaultValue="false" Description="Build and test Win64"/>
 
+<Option Name="AndroidCookFlavor" DefaultValue="ETC2" Description="Cook flavor to use for build/cook/run when target platform is Android"/>
+
 <Do If="'$(EditorPlatforms)' == ''" >
 	<Property Name="EditorPlatforms" Value="" />
 	<Property Name="EditorPlatforms" Value="$(EditorPlatforms)Win64+" If="$(WithWin64Editor)" />
