Follow-on change for CL 40037089 ensuring that the cook stage also uses the updated AndroidCookFlavor.

diff --git a/Engine/Build/Graph/Tasks/BuildAndTestProject.xml b/Engine/Build/Graph/Tasks/BuildAndTestProject.xml
--- a/Engine/Build/Graph/Tasks/BuildAndTestProject.xml
+++ b/Engine/Build/Graph/Tasks/BuildAndTestProject.xml
@@ -585,7 +585,7 @@
 		<!-- Set CookPlatform since it isn't always the same as the target -->
 		<Property Name="CookPlatform" Value="$(TargetPlatform)" />
 		<Property Name="CookPlatform" Value="Windows" If="'$(TargetPlatform)' == 'Win64'"/>
-		<Property Name="CookPlatform" Value="$(CookPlatform)_ETC2" If="'$(TargetPlatform)' == 'Android'"/>
+		<Property Name="CookPlatform" Value="$(CookPlatform)_$(AndroidCookFlavor)" If="'$(TargetPlatform)' == 'Android'"/>
 		<Property Name="CookPlatform" Value="$(CookPlatform)Client" If="'$(TargetType)' == 'Client'" />
 
 		<!-- Only use Cook Flavor for Android -->
