Remove UE-CloudPublishNamespace from stream configuration, and remove checks for it when attempting to publish snapshots.

[RN] none

diff --git a/Engine/Build/Graph/Tasks/BuildAndTestProject.xml b/Engine/Build/Graph/Tasks/BuildAndTestProject.xml
--- a/Engine/Build/Graph/Tasks/BuildAndTestProject.xml
+++ b/Engine/Build/Graph/Tasks/BuildAndTestProject.xml
@@ -693,7 +693,7 @@
 			<!-- (Optional) Publish from local Zen to Cloud or remote Zen server -->
 			<Property Name="PlatformPublishStageDependencyName" Value="$(PlatformHashNodeName)" />
 			<Property Name="ShouldSnapshotPublishCloud" Value="false"/>
-			<Property Name="ShouldSnapshotPublishCloud" Value="true" If="$(PublishCloudSnapshot) and '$(UE-CloudPublishHost)' != '' and '$(UE-CloudPublishNamespace)' != ''"/>
+			<Property Name="ShouldSnapshotPublishCloud" Value="true" If="$(PublishCloudSnapshot) and '$(UE-CloudPublishHost)' != ''"/>
 			<Property Name="ShouldSnapshotPublishZen" Value="false"/>
 			<Property Name="ShouldSnapshotPublishZen" Value="true" If="$(IsBuildMachine) and !$(IsPreflight) and $(PublishZenSnapshot) and '$(UE-ZenPublishHost)' != ''"/>
 			<Property Name="ShouldSnapshotPublish" Value="false"/>
