Renamed option to publish cloud artifacts to PublishCloudArtifacts and RetrieveUsingCloudArtifacts to avoid exposing implementation detail that these are hosted by [REDACTED_PERSON] (Jupiter)

diff --git a/Engine/Build/Graph/Tasks/BuildAndTestProject.xml b/Engine/Build/Graph/Tasks/BuildAndTestProject.xml
--- a/Engine/Build/Graph/Tasks/BuildAndTestProject.xml
+++ b/Engine/Build/Graph/Tasks/BuildAndTestProject.xml
@@ -17,8 +17,8 @@
 	<Property Name="IsHordeEnv" Value="true" If="'$(UE_HORDE_JOBID)' != ''"/>
 	<Option Name="Horde" DefaultValue="$(IsBuildMachine) And $(IsHordeEnv)" Description="Whether the build is running under Horde"/>
 	
-	<Option Name="PublishToCloudDDC" DefaultValue="false" Description="Enable to publish to cloud ddc instead of fileshare" />
-	<Option Name="RetrieveFromCloudDDC" DefaultValue="false" Description="Enable to retrieve artifacts from cloud ddc" />
+	<Option Name="PublishCloudArtifacts" DefaultValue="false" Description="Enable to publish cloud artifacts" />
+	<Option Name="RetrieveUsingCloudArtifacts" DefaultValue="false" Description="Enable to retrieve cloud artifacts" />
 
 	<!-- Project Options -->
 	<Option Name="ProjectFile" DefaultValue="" Restrict=".*" Description="Path to the project file. Can be used instead of specifying ProjectName and ProjectPath separately."/>
@@ -458,8 +458,8 @@
 				<Property Name="TestPrerequisites" Value="$(TestPrerequisites);$(HordeArtifactPlatformRetrieve)" />
 				<Property Name="CookedEditorBuildPath" Value="$(ProjectOutputDirectory)/$(EditorPlatformPath)CookedEditor" />
 			</Do>
-			<Do If="$(RetrieveFromCloudDDC)">
-				<Property Name="CloudArtifactPlatformRetrieve" Value="$(PreNodeName)Retrieve CookedEditor Build Artifact $(EditorPlatform) (CloudDDC)"/>
+			<Do If="$(RetrieveUsingCloudArtifacts)">
+				<Property Name="CloudArtifactPlatformRetrieve" Value="$(PreNodeName)Retrieve CookedEditor Build Artifact $(EditorPlatform) (Cloud)"/>
 				<Node Name="$(CloudArtifactPlatformRetrieve)" Requires="$(TestPrerequisites);$(PlatformPublishNodeName)">
 					<RetrieveCloudArtifact Name="$(BuildNamePath)" Type="staged-build" OutputDir="$(ProjectOutputDirectory)/$(EditorPlatformPath)CookedEditor" Project="$(TargetProject)" Platform="$(EditorPlatform)" />
 				</Node>
@@ -579,18 +579,18 @@
 		<Property Name="PlatformPublishPackagedNodeName" Value="$(PreNodeNameWithPlatform)Publish Packaged $(TargetPlatform)" />
 		<Property Name="PlatformCheckBuildSizeNodeName" Value="$(PreNodeNameWithPlatform)Check Build Size $(TargetPlatform)" />
 		<Property Name="PlatformStagedArtifactNodeName" Value="$(PlatformPublishStagedNodeName) (Horde)" />
-		<Property Name="CloudPlatformStagedArtifactNodeName" Value="$(PlatformPublishStagedNodeName) (CloudDDC)" />
+		<Property Name="CloudPlatformStagedArtifactNodeName" Value="$(PlatformPublishStagedNodeName) (Cloud)" />
 		<Property Name="HordeArtifactPlatformPackagedNodeName" Value="$(PlatformPublishPackagedNodeName) (Horde)" />
 		<Property Name="HordeArtifactPlatformRetrieve" Value="$(PreNodeNameWithPlatform)Retrieve Build Artifact $(TargetPlatform)" />
-		<Property Name="CloudDDCArtifactPlatformPackagedNodeName" Value="$(PlatformPublishPackagedNodeName) (CloudDDC)" />
-		<Property Name="CloudDDCArtifactPlatformRetrieve" Value="$(PreNodeNameWithPlatform)Retrieve Build Artifact $(TargetPlatform)" />
+		<Property Name="CloudArtifactPlatformPackagedNodeName" Value="$(PlatformPublishPackagedNodeName) (Cloud)" />
+		<Property Name="CloudArtifactPlatformRetrieve" Value="$(PreNodeNameWithPlatform)Retrieve Build Artifact $(TargetPlatform)" />
 		<Property Name="RequirePlatformStaged" Value="$(PlatformStageNodeName)" />
 		<Property Name="RequirePlatformStaged" Value="$(RequirePlatformStaged);$(PlatformStagedArtifactNodeName)" If="$(Horde)"/>
-		<Property Name="RequirePlatformStaged" Value="$(RequirePlatformStaged);$(CloudPlatformStagedArtifactNodeName)" If="$(PublishToCloudDDC)"/>
+		<Property Name="RequirePlatformStaged" Value="$(RequirePlatformStaged);$(CloudPlatformStagedArtifactNodeName)" If="$(PublishCloudArtifacts)"/>
 		<Property Name="RequirePlatformStaged" Value="$(RequirePlatformStaged);$(PlatformPublishStagedNodeName)" If="!$(SkipPublish) And $(PublishBuilds)"/>
 		<Property Name="RequirePlatformPackaged" Value="$(PlatformPackageNodeName)" />
 		<Property Name="RequirePlatformPackaged" Value="$(RequirePlatformPackaged);$(HordeArtifactPlatformPackagedNodeName)" If="$(Horde)"/>
-		<Property Name="RequirePlatformPackaged" Value="$(RequirePlatformPackaged);$(CloudDDCArtifactPlatformPackagedNodeName)" If="$(PublishToCloudDDC)"/>
+		<Property Name="RequirePlatformPackaged" Value="$(RequirePlatformPackaged);$(CloudArtifactPlatformPackagedNodeName)" If="$(PublishCloudArtifacts)"/>
 		<Property Name="RequirePlatformPackaged" Value="$(RequirePlatformPackaged);$(PlatformPublishPackagedNodeName)"  If="!$(SkipPublish) And $(PublishBuilds)"/>
 		
 		<!-- Set CookPlatform since it isn't always the same as the target -->
@@ -754,8 +754,8 @@
 				<Copy Files="..." From="$(ShaderSymbolsDirectory)/*$(TargetPlatform)*" To="$(NetworkOutputDirectory)/$(StagedPlatformFolder)/ShaderSymbols"/>
 			</Node>
 
-			<!-- Publish (Staged) (CloudDDC) -->
-			<Do If="$(PublishToCloudDDC)">
+			<!-- Publish (Staged) (Cloud) -->
+			<Do If="$(PublishCloudArtifacts)">
 				<Node Name="$(CloudPlatformStagedArtifactNodeName)" Requires="$(PlatformStageNodeName)">
 					<ForEach Name="PackageConfiguration" Values="$(PackageConfigurations)" Separator="+">
 						<Tag Files="$(ProjectOutputDirectory)/$(StagedPlatformFolder)/..." With="#$(PlatformStageNodeName) Files"/>
@@ -805,12 +805,12 @@
 				</Node>
 			</Do>
 
-			<!-- Publish (Packages) (CloudDDC) -->
-			<Do  If="$(PublishToCloudDDC)">
+			<!-- Publish (Packages) (Cloud) -->
+			<Do  If="$(PublishCloudArtifacts)">
 				<Property Name="CloudPackageSkipBCRArguments" Value="-stage -package"/>
 				<Property Name="CloudPackageSkipBCRArguments" Value="-skipstage -skippackage" If="'$(TargetPlatform)' == 'Mac'"/>
 
-				<Node Name="$(CloudDDCArtifactPlatformPackagedNodeName)" Requires="$(PlatformPackageNodeName)">
+				<Node Name="$(CloudArtifactPlatformPackagedNodeName)" Requires="$(PlatformPackageNodeName)">
 					<ForEach Name="PackageConfiguration" Values="$(PackageConfigurations)" Separator="+">
 						<Command Name="BuildCookRun" Arguments="$(BCRArgs) -configuration=$(PackageConfiguration) $(StageAndPackageArguments) -skipbuild -skipcook -skippak $(CloudPackageSkipBCRArguments) -archive -archivedirectory=&quot;$(ProjectOutputDirectory)/Packages/$(StagedPlatformFolder)/$(PackageConfiguration)&quot;" />
 						<Tag Files="$(ProjectOutputDirectory)/Packages/$(StagedPlatformFolder)/$(PackageConfiguration)/..." With="#$(PlatformPackageNodeName) $(PackageConfiguration) Files"/>
