Avoid use of the dot character when composing a bucket name to export snapshots to.

diff --git a/Engine/Build/Graph/Tasks/BuildAndTestProject.xml b/Engine/Build/Graph/Tasks/BuildAndTestProject.xml
--- a/Engine/Build/Graph/Tasks/BuildAndTestProject.xml
+++ b/Engine/Build/Graph/Tasks/BuildAndTestProject.xml
@@ -711,10 +711,12 @@
 					<Property Name="SnapshotLongTermStorageDir" Value="$(NetworkSnapshotDescFileRoot)/$(EscapedBranch)/$(ProjectName)/$(SnapshotLongTermPlatformFolder)/$(Change)"  If="'$(NetworkSnapshotDescFileRoot)' != ''"/>
 
 					<Do If="$(ShouldSnapshotPublishCloud)">
+						<Property Name="SnapshotBranch" Value="$(EscapedBranch)"/>
+						<StringOp Input="$(EscapedBranch)" Method="Replace" Arguments=".;-" Output="SnapshotBranch"/>
 						<!-- Publish to generic cloud storage endpoint -->
 						<Property Name="SnapshotBaseDescriptorFile" Value=""/>
 						<Property Name="SnapshotBaseDescriptorFile" Value="$(NetworkSnapshotDescFileRoot)/$(EscapedBranch)/$(ProjectName)/$(StagedPlatformFolder)/$(UE_HORDE_LAST_WARNING_CL)/$(SnapshotFilenamePrefix)-cloud.json"  If="'$(NetworkSnapshotDescFileRoot)' != '' and '$(UE_HORDE_LAST_WARNING_CL)' != ''"/>
-						<ZenExportSnapshot Project="$(TargetProject)" Platform="$(CookPlatform)" DestinationStorageType="Cloud" DestinationCloudHost="$(UE-CloudPublishHost)" DestinationCloudHttpVersion="$(UE-CloudPublishHttpVersion)" SnapshotDescriptorCloudHost="$(UE-CloudPublishDescriptorHost)" SnapshotDescriptorCloudHttpVersion="$(UE-CloudPublishDescriptorHttpVersion)" DestinationCloudNamespace="$(UE-CloudPublishNamespace)" DestinationCloudBucket="$(ProjectName).oplog.$(EscapedBranch).$(CookPlatform)" DestinationIdentifier="$(SnapshotIdentifier)" SnapshotDescriptorFile="$(SnapshotLocalDir)/$(StagedPlatformFolder)/$(SnapshotFilenamePrefix)-cloud.json" SnapshotBaseDescriptorFile="$(SnapshotBaseDescriptorFile)"/>
+						<ZenExportSnapshot Project="$(TargetProject)" Platform="$(CookPlatform)" DestinationStorageType="Cloud" DestinationCloudHost="$(UE-CloudPublishHost)" DestinationCloudHttpVersion="$(UE-CloudPublishHttpVersion)" SnapshotDescriptorCloudHost="$(UE-CloudPublishDescriptorHost)" SnapshotDescriptorCloudHttpVersion="$(UE-CloudPublishDescriptorHttpVersion)" DestinationCloudNamespace="$(UE-CloudPublishNamespace)" DestinationCloudBucket="$(ProjectName).oplog.$(SnapshotBranch).$(CookPlatform)" DestinationIdentifier="$(SnapshotIdentifier)" SnapshotDescriptorFile="$(SnapshotLocalDir)/$(StagedPlatformFolder)/$(SnapshotFilenamePrefix)-cloud.json" SnapshotBaseDescriptorFile="$(SnapshotBaseDescriptorFile)"/>
 						<Copy From="$(SnapshotLocalDir)/$(StagedPlatformFolder)/$(SnapshotFilenamePrefix)-cloud.json" To="$(NetworkOutputDirectory)/$(StagedPlatformFolder)/SnapshotDescriptor.json"/>
 						<Do If="'$(SnapshotLongTermStorageDir)' != ''">
 							<Copy From="$(SnapshotLocalDir)/$(StagedPlatformFolder)/$(SnapshotFilenamePrefix)-cloud.json" To="$(SnapshotLongTermStorageDir)/$(SnapshotFilenamePrefix)-cloud.json"/>
@@ -727,7 +729,7 @@
 						<Property Name="SnapshotBuildMetaData" Value="$(SnapshotBuildMetaData);preflightChangelist=$(PreflightChange)" If="'$(PreflightChange)' != ''" />
 						<Property Name="SnapshotBuildMetaData" Value="$(SnapshotBuildMetaData);hordeTemplateId=$(UE_HORDE_TEMPLATEID)" If="'$(UE_HORDE_TEMPLATEID)' != ''" />
 						<Log Message="Publishing build with metadata '$(SnapshotBuildMetaData)'"/>
-						<ZenExportSnapshot Project="$(TargetProject)" Platform="$(CookPlatform)" DestinationStorageType="Builds" DestinationCloudHost="$(UE-CloudPublishHost)" DestinationCloudHttpVersion="$(UE-CloudPublishHttpVersion)" SnapshotDescriptorCloudHost="$(UE-CloudPublishDescriptorHost)" SnapshotDescriptorCloudHttpVersion="$(UE-CloudPublishDescriptorHttpVersion)" DestinationCloudNamespace="$(UE-CloudPublishNamespace)" DestinationCloudBucket="$(ProjectName).oplog.$(EscapedBranch).$(CookPlatform)" DestinationIdentifier="$(SnapshotIdentifier)" SnapshotDescriptorFile="$(SnapshotLocalDir)/$(StagedPlatformFolder)/$(SnapshotFilenamePrefix)-builds.json" Metadata="$(SnapshotBuildMetaData)"/>
+						<ZenExportSnapshot Project="$(TargetProject)" Platform="$(CookPlatform)" DestinationStorageType="Builds" DestinationCloudHost="$(UE-CloudPublishHost)" DestinationCloudHttpVersion="$(UE-CloudPublishHttpVersion)" SnapshotDescriptorCloudHost="$(UE-CloudPublishDescriptorHost)" SnapshotDescriptorCloudHttpVersion="$(UE-CloudPublishDescriptorHttpVersion)" DestinationCloudNamespace="$(UE-CloudPublishNamespace)" DestinationCloudBucket="$(ProjectName).oplog.$(SnapshotBranch).$(CookPlatform)" DestinationIdentifier="$(SnapshotIdentifier)" SnapshotDescriptorFile="$(SnapshotLocalDir)/$(StagedPlatformFolder)/$(SnapshotFilenamePrefix)-builds.json" Metadata="$(SnapshotBuildMetaData)"/>
 						<Copy From="$(SnapshotLocalDir)/$(StagedPlatformFolder)/$(SnapshotFilenamePrefix)-builds.json" To="$(NetworkOutputDirectory)/$(StagedPlatformFolder)/SnapshotDescriptorBuilds.json"/>
 						<Do If="'$(SnapshotLongTermStorageDir)' != ''">
 							<Copy From="$(SnapshotLocalDir)/$(StagedPlatformFolder)/$(SnapshotFilenamePrefix)-builds.json" To="$(SnapshotLongTermStorageDir)/$(SnapshotFilenamePrefix)-builds.json"/>
