Add ability to define additional Editor Targets in the BuildEditorAndTools buildgraph
#rnx

diff --git a/Engine/Build/Graph/Examples/BuildEditorAndTools.xml b/Engine/Build/Graph/Examples/BuildEditorAndTools.xml
--- a/Engine/Build/Graph/Examples/BuildEditorAndTools.xml
+++ b/Engine/Build/Graph/Examples/BuildEditorAndTools.xml
@@ -44,6 +44,7 @@
 	<Option Name="ToolCompileArgs" DefaultValue="" Description="Arguments to be used for Tool Compile"/>
 	<Option Name="GameTargets" Restrict="[^ ]*" DefaultValue="" Description="Names of game targets to build separated by semicolons, eg. UnrealClient;UnrealServer"/>
 	<Option Name="AdditionalTools" Restrict="[^ ]*" DefaultValue="" Description="Tool Compiles to build seperated by semicolons, eg. UnrealMultiUserServer;BuildPatchTool"/>
+	<Option Name="AdditionalEditors" Restrict="[^ ]*" DefaultValue="" Description="Editor Compiles to build seperated by semicolons"/>
 	<Option Name="TargetPlatforms" Restrict="[^ ]*" DefaultValue="Win64" Description="List of the target platforms to build for, separated by semicolons, eg. Win64;Android"/>
 	<Option Name="OutputDir" DefaultValue="$(RootDir)\LocalBuilds\$(EditorTarget)Binaries" Description ="Output directory for compiled binaries"/>
 	<Option Name="Licensee" Restrict="true|false" DefaultValue="true" Description="Whether to treat the changelist number as being from a licensee Perforce server"/>
@@ -91,6 +92,9 @@
 		<Node Name="Compile $(EditorTarget) Win64" Produces="#EditorBinaries;#TargetInfo">
 			<Compile Target="$(EditorTarget)" Platform="Win64" Configuration="Development" Project="$(UProjectPath)" Tag="#EditorBinaries" Arguments="$(EditorCompileArgs)"/>
 			<QueryTargets ProjectFile="$(UProjectPath)" Tag="#TargetInfo"/>
+			<ForEach Name="AdditonalEditor" Values="$(AdditionalEditors)">
+				<Compile Target="$(AdditonalEditor)" Platform="Win64" Configuration="Development" Tag="#EditorBinaries" Arguments="$(EditorCompileArgs)"/>
+			</ForEach>
 		</Node>
 		
 		<!-- Compile the game targets -->
