Stop publishing snapshots with the older generic cloud storage mechanism.  Now relying exclusively on the builds API.

diff --git a/Engine/Build/Graph/Tasks/BuildAndTestProject.xml b/Engine/Build/Graph/Tasks/BuildAndTestProject.xml
--- a/Engine/Build/Graph/Tasks/BuildAndTestProject.xml
+++ b/Engine/Build/Graph/Tasks/BuildAndTestProject.xml
@@ -713,14 +713,6 @@
 					<Do If="$(ShouldSnapshotPublishCloud)">
 						<Property Name="SnapshotBranch" Value="$(EscapedBranch)"/>
 						<StringOp Input="$(EscapedBranch)" Method="Replace" Arguments=".;-" Output="SnapshotBranch"/>
-						<!-- Publish to generic cloud storage endpoint -->
-						<Property Name="SnapshotBaseDescriptorFile" Value=""/>
-						<Property Name="SnapshotBaseDescriptorFile" Value="$(NetworkSnapshotDescFileRoot)/$(EscapedBranch)/$(ProjectName)/$(StagedPlatformFolder)/$(UE_HORDE_LAST_WARNING_CL)/$(SnapshotFilenamePrefix)-cloud.json"  If="'$(NetworkSnapshotDescFileRoot)' != '' and '$(UE_HORDE_LAST_WARNING_CL)' != ''"/>
-						<ZenExportSnapshot Project="$(TargetProject)" Platform="$(CookPlatform)" DestinationStorageType="Cloud" DestinationCloudHost="$(UE-CloudPublishHost)" DestinationCloudHttpVersion="$(UE-CloudPublishHttpVersion)" SnapshotDescriptorCloudHost="$(UE-CloudPublishDescriptorHost)" SnapshotDescriptorCloudHttpVersion="$(UE-CloudPublishDescriptorHttpVersion)" DestinationCloudNamespace="$(UE-CloudPublishNamespace)" DestinationCloudBucket="$(ProjectName).oplog.$(SnapshotBranch).$(CookPlatform)" DestinationIdentifier="$(SnapshotIdentifier)" SnapshotDescriptorFile="$(SnapshotLocalDir)/$(StagedPlatformFolder)/$(SnapshotFilenamePrefix)-cloud.json" SnapshotBaseDescriptorFile="$(SnapshotBaseDescriptorFile)"/>
-						<Copy From="$(SnapshotLocalDir)/$(StagedPlatformFolder)/$(SnapshotFilenamePrefix)-cloud.json" To="$(NetworkOutputDirectory)/$(StagedPlatformFolder)/SnapshotDescriptor.json"/>
-						<Do If="'$(SnapshotLongTermStorageDir)' != ''">
-							<Copy From="$(SnapshotLocalDir)/$(StagedPlatformFolder)/$(SnapshotFilenamePrefix)-cloud.json" To="$(SnapshotLongTermStorageDir)/$(SnapshotFilenamePrefix)-cloud.json"/>
-						</Do>
 						<!-- Publish to builds cloud storage endpoint -->
 						<Property Name="SnapshotBuildBaselineBranch" Value="$(Branch)"/>
 						<Property Name="SnapshotBuildBaselineBranch" Value="$(UE-CloudPublishBaselineBranch)" If="'$(UE-CloudPublishBaselineBranch)' != ''"/>
@@ -730,7 +722,7 @@
 						<Property Name="SnapshotBuildMetaData" Value="$(SnapshotBuildMetaData);hordeTemplateId=$(UE_HORDE_TEMPLATEID)" If="'$(UE_HORDE_TEMPLATEID)' != ''" />
 						<Log Message="Publishing build with metadata '$(SnapshotBuildMetaData)'"/>
 						<ZenExportSnapshot Project="$(TargetProject)" Platform="$(CookPlatform)" DestinationStorageType="Builds" DestinationCloudHost="$(UE-CloudPublishHost)" DestinationCloudHttpVersion="$(UE-CloudPublishHttpVersion)" SnapshotDescriptorCloudHost="$(UE-CloudPublishDescriptorHost)" SnapshotDescriptorCloudHttpVersion="$(UE-CloudPublishDescriptorHttpVersion)" DestinationCloudNamespace="$(UE-CloudPublishNamespace)" DestinationCloudBucket="$(ProjectName).oplog.$(SnapshotBranch).$(CookPlatform)" DestinationIdentifier="$(SnapshotIdentifier)" SnapshotDescriptorFile="$(SnapshotLocalDir)/$(StagedPlatformFolder)/$(SnapshotFilenamePrefix)-builds.json" Metadata="$(SnapshotBuildMetaData)"/>
-						<Copy From="$(SnapshotLocalDir)/$(StagedPlatformFolder)/$(SnapshotFilenamePrefix)-builds.json" To="$(NetworkOutputDirectory)/$(StagedPlatformFolder)/SnapshotDescriptorBuilds.json"/>
+						<Copy From="$(SnapshotLocalDir)/$(StagedPlatformFolder)/$(SnapshotFilenamePrefix)-builds.json" To="$(NetworkOutputDirectory)/$(StagedPlatformFolder)/SnapshotDescriptor.json"/>
 						<Do If="'$(SnapshotLongTermStorageDir)' != ''">
 							<Copy From="$(SnapshotLocalDir)/$(StagedPlatformFolder)/$(SnapshotFilenamePrefix)-builds.json" To="$(SnapshotLongTermStorageDir)/$(SnapshotFilenamePrefix)-builds.json"/>
 						</Do>
