Fixed issue with horde artifacts not being downloaded by default due to incorrect checks for RetrieveUsingCloudArtifacts (which should be of by default)

diff --git a/Engine/Build/Graph/Tasks/BuildAndTestProject.xml b/Engine/Build/Graph/Tasks/BuildAndTestProject.xml
--- a/Engine/Build/Graph/Tasks/BuildAndTestProject.xml
+++ b/Engine/Build/Graph/Tasks/BuildAndTestProject.xml
@@ -18,7 +18,7 @@
 	<Option Name="Horde" DefaultValue="$(IsBuildMachine) And $(IsHordeEnv)" Description="Whether the build is running under Horde"/>
 	
 	<Option Name="PublishCloudArtifacts" DefaultValue="false" Description="Enable to publish cloud artifacts" />
-	<Option Name="RetrieveUsingCloudArtifacts" DefaultValue="false" Description="Enable to retrieve cloud artifacts" />
+	<Option Name="RetrieveUsingCloudArtifacts" DefaultValue="false" Restrict="true|false" Description="Enable to retrieve cloud artifacts" />
 
 	<!-- Project Options -->
 	<Option Name="ProjectFile" DefaultValue="" Restrict=".*" Description="Path to the project file. Can be used instead of specifying ProjectName and ProjectPath separately."/>
@@ -852,7 +852,7 @@
 					<Property Name="GauntletBuildArg" Value="$(GauntletBuildArg) -destlocalinstalldir=&quot;$(LocalInstallDirectory)&quot;" If="!$(IsLocalPlatform) And $(WithDeviceReservation)"/>
 				</Do>
 
-				<Do If="$(Horde) and $(RetrieveUsingCloudArtifacts) == ''">
+				<Do If="$(Horde) and '$(RetrieveUsingCloudArtifacts)' == 'false'">
 					<Property Name="HordeArtifactPlatformRetrieveWithConfig" Value="$(HordeArtifactPlatformRetrieve) $(TestConfiguration) (Horde)"/>
 					<Node Name="$(HordeArtifactPlatformRetrieveWithConfig)" Requires="$(TestPrerequisites)">
 						<Switch>
@@ -867,7 +867,7 @@
 					<Property Name="TestPrerequisites" Value="$(TestPrerequisites);$(HordeArtifactPlatformRetrieveWithConfig)" />
 				</Do>
 
-				<Do If="$(RetrieveUsingCloudArtifacts)">
+				<Do If="'$(RetrieveUsingCloudArtifacts)' == 'true'">
 					<Property Name="CloudArtifactPlatformRetrieveWithConfig" Value="$(CloudArtifactPlatformRetrieve) $(TestConfiguration) (Cloud)"/>
 					<Node Name="$(CloudArtifactPlatformRetrieveWithConfig)" Requires="$(TestPrerequisites)">
 						<Switch>
