Added option to publish PCBs to cloud storage as well as converted the horde upload to use a task instead which gives us clearer reporting.
This adds a new target "Submit PCBs" which will upload pcbs to all backends and not just Perforce. The old perforce node still exists for backwards compatability but its recommended to use the new one instead.

diff --git a/Engine/Build/Graph/Examples/BuildEditorAndTools.xml b/Engine/Build/Graph/Examples/BuildEditorAndTools.xml
--- a/Engine/Build/Graph/Examples/BuildEditorAndTools.xml
+++ b/Engine/Build/Graph/Examples/BuildEditorAndTools.xml
@@ -55,11 +55,27 @@
 	<Option Name="PreferredAgent" DefaultValue="CompileWin64;Win64" Description ="Agent type to use for job"/>
 	<Option Name="ArchiveName" Restrict="[^ ]*" DefaultValue="$(EditorTarget)" Description="Names of archive to generate instead of using the GameTarget, when the GameTarget may be UnrealEditor and result in a clash"/>
 	<Option Name="SymbolStorePath" DefaultValue="" Description ="Output directory for symbol storage"/>
+	<Option Name="PublishCloudArtifacts" DefaultValue="false" Restrict="true|false"  Description="Enable to publish cloud artifacts" />
 
 	<Property Name="ArchiveDir" Value="$(RootDir)\LocalBuilds\ArchiveForUGS"/>
 	<Property Name="ArchiveStagingDir" Value="$(ArchiveDir)\Staging"/>
 	<Property Name="ArchiveStagingRelativeDir" Value="LocalBuilds\ArchiveForUGS\Staging"/>
 
+	<EnvVar Name="UE_HORDE_JOBID"/>
+
+	<Property Name="IsHordeEnv" Value="false"/>
+	<Property Name="IsHordeEnv" Value="true" If="'$(UE_HORDE_JOBID)' != ''"/>
+	<Option Name="Horde" DefaultValue="$(IsBuildMachine) And $(IsHordeEnv)" Description="Whether the build is running under Horde"/>
+
+	<Property Name="PerforceSubmitNodeName" Value="Submit To Perforce For UGS" />
+	<Property Name="CloudSubmitNodeName" Value="Submit To Cloud Storage For UGS" />
+	<Property Name="HordeSubmitNodeName" Value="Submit To Horde Storage For UGS" />
+
+	<Property Name="GeneratedSubmitNodes" Value="" />
+	<Property Name="GeneratedSubmitNodes" Value="$(GeneratedSubmitNodes);$(HordeSubmitNodeName)" If="$(Horde)"/>
+	<Property Name="GeneratedSubmitNodes" Value="$(GeneratedSubmitNodes);$(CloudSubmitNodeName)" If="$(PublishCloudArtifacts)"/>
+	<Property Name="GeneratedSubmitNodes" Value="$(GeneratedSubmitNodes);$(PerforceSubmitNodeName)" If="'$(ArchiveStream)' != ''"/>
+
 	<EnvVar Name="COMPUTERNAME"/>
 	<EnvVar Name="P4CLIENT"/>
 	
@@ -168,7 +184,7 @@
 		</Node>
 		
 		<!-- Create a zip file and submit it to Perforce -->
-		<Node Name="Submit To Perforce For UGS" Requires="Stage for UGS">
+		<Node Name="$(PerforceSubmitNodeName)" Requires="Stage for UGS">
 			<Property Name="ArchivePerforceDir" Value="$(RootDir)\LocalBuilds\ArchiveForUGS-Perforce"/>
 			<Property Name="ArchiveFile" Value="$(ArchivePerforceDir)\$(EscapedBranch)-$(ArchiveName).zip"/>
 			<Zip FromDir="$(ArchiveStagingDir)" ZipFile="$(ArchiveFile)"/>
@@ -178,14 +194,24 @@
 			<Submit Description="[CL $(CodeChange)] Updated binaries" Files="$(ArchiveFile)" FileType="binary+FS32" Workspace="$(SubmitClient)" Stream="$(ArchiveStream)" RootDir="$(ArchivePerforceDir)" Force="$(ForceSubmit)" If="'$(ArchiveStream)' != ''"/>
 		</Node>
 
-		<!-- Upload binaries to Horde -->
-		<Property Name="UgsProjectKey" Value=""/>
-		<Property Name="UgsProjectKey" Value="ugs-project=$(Branch)/$(UProjectPath)" If="'$(Branch)' != '' and '$(UProjectPath)' != ''"/>
-		<Property Name="UgsProjectKey" Value="ugs-project=$(Branch)" If="'$(Branch)' != '' and '$(UgsProjectKey)' == ''"/>
-		<Artifact Name="editor" Description="Editor PCBs" Type="ugs-pcb" BasePath="$(ArchiveStagingRelativeDir)" Keys="$(UgsProjectKey)" Metadata="ArchiveType=Editor" Tag="#PublishBinaries"/>
-		
+		<Node Name="$(CloudSubmitNodeName)" Requires="Stage for UGS">
+			<Error Message="UProjectPath is a required option when submitting to Cloud Storage" If="'$(UProjectPath)' == ''" />
+			<Property Name="UgsProjectKey" Value="ugs-project=$(Branch)/$(UProjectPath)" />
+
+			<CreateCloudArtifact Name="editor" Type="ugs-pcb" BaseDir="$(ArchiveStagingRelativeDir)" Metadata="$(UgsProjectKey);ArchiveType=Editor" Project="$(UProjectPath)" Platform="Windows" Files="#PublishBinaries"/>
+		</Node>
+
+		<Node Name="$(HordeSubmitNodeName)" Requires="Stage for UGS">
+			<!-- Upload binaries to Horde -->
+			<Property Name="UgsProjectKey" Value=""/>
+			<Property Name="UgsProjectKey" Value="ugs-project=$(Branch)/$(UProjectPath)" If="'$(Branch)' != '' and '$(UProjectPath)' != ''"/>
+			<Property Name="UgsProjectKey" Value="ugs-project=$(Branch)" If="'$(Branch)' != '' and '$(UgsProjectKey)' == ''"/>
+			<CreateArtifact Name="editor" Description="Editor PCBs" Type="ugs-pcb" BaseDir="$(ArchiveStagingRelativeDir)" Keys="$(UgsProjectKey)" Metadata="ArchiveType=Editor" Files="#PublishBinaries"/>		
+		</Node>
 	</Agent>
 
+	<Aggregate Name="Submit PCBs" Requires="$(GeneratedSubmitNodes)"/>
+
 	<Label Category="Editors" Name="Win64" Requires="Compile $(EditorTarget) Win64" UgsBadge="Editor"/>
 
 	<ForEach Name="TargetPlatform" Values="$(TargetPlatforms)">
