Avoid using nested variables in TestConfigs + enable using Cloud build artifacts retrieval by default in Epic env

#rnx

diff --git a/Engine/Build/Graph/Tasks/BuildAndTestProject.xml b/Engine/Build/Graph/Tasks/BuildAndTestProject.xml
--- a/Engine/Build/Graph/Tasks/BuildAndTestProject.xml
+++ b/Engine/Build/Graph/Tasks/BuildAndTestProject.xml
@@ -17,9 +17,6 @@
 	<Property Name="IsHordeEnv" Value="true" If="'$(UE_HORDE_JOBID)' != ''"/>
 	<Option Name="Horde" DefaultValue="$(IsBuildMachine) And $(IsHordeEnv)" Description="Whether the build is running under Horde"/>
 	
-	<Option Name="PublishCloudArtifacts" DefaultValue="false" Description="Enable to publish cloud artifacts" />
-	<Option Name="RetrieveUsingCloudArtifacts" DefaultValue="false" Restrict="true|false" Description="Enable to retrieve cloud artifacts" />
-
 	<!-- Project Options -->
 	<Option Name="ProjectFile" DefaultValue="" Restrict=".*" Description="Path to the project file. Can be used instead of specifying ProjectName and ProjectPath separately."/>
 
@@ -170,7 +167,11 @@
 		<Include Script="$(PreStepsExtendFile)"/>
 	</Do>
 
-	<Option Name="UseTestConfigs" DefaultValue="false" Description="Use Test Configs (see RunTestConfigs.xml) instead of the standard test node generation"/>
+	<!-- Test configs -->
+	<Include Script="Inc/TestConfigs.xml"/>
+	<Do If="'$(TestConfigList)' == ''">
+		<Expand Name="AddTestConfig" TestConfigName=""/>
+	</Do>
 
 	<!-- At this point we have all options so calculate a few things -->
 
@@ -208,10 +209,13 @@
 	<!-- On Epic build machines, we want to:
 		1. Write dependency information into the cooked output so that it is suitable for use as a base upon which to cook incrementally.
 		2. Cook incrementally if the initial state at the start of the job is not empty/blank.
+		3. Use Cloud artifacts by default to publish and retrieve build
 	-->
 	<Do If="$(IsEpicBuildMachine)">
 		<Property Name="ExtraCookArguments" Value="$(ExtraCookArguments) -cookincremental"/>
 	</Do>
+	<Option Name="RetrieveUsingCloudArtifacts" DefaultValue="$(IsEpicBuildMachine)" Restrict="true|false" Description="Enable to retrieve cloud artifacts" />
+	<Option Name="PublishCloudArtifacts" DefaultValue="$(RetrieveUsingCloudArtifacts)" Description="Enable to publish cloud artifacts" />
 
 	<!-- Detect the type of the client based on the target name -->
 	<Property Name="TargetType" Value="Game" />
@@ -221,10 +225,7 @@
 	<Macro Name="SanitizeTestNameTitle">
 		<!-- Expect a Property named "TestNameTitle" to be defined beforehand -->
 		<Do If="Contains('$(TestNameTitle)', ':')">
-			<Property Name="Left" Value=""/>
-			<Property Name="Right" Value=""/>
-			<Regex Pattern="(.*):(.*)" Capture="Left;Right" Input="$(TestNameTitle)"/>
-			<Property Name="TestNameTitle" Value="$(Left)_$(Right)"/>
+			<StringOp Input="$(TestNameTitle)" Method="Replace" Arguments=":;_" Output="TestNameTitle"/>
 		</Do>
 		<!-- We want to use a user-provided name if any. Expected format of the node name is 'FooTest(arg1+arg2+arg3)[Test Name]' and we want to use Test Name if provided -->
 		<!-- This is intentionally after the previous test to override it. -->
@@ -424,26 +425,28 @@
 
 			<!-- Publish CookedEditor -->
 			<Node Name="$(PlatformPublishNodeName)" Requires="$(PreNodeName)Stage CookedEditor $(EditorPlatform)">
-				<Switch>
-					<Case If="$(Horde)">
-						<Expand Name="WriteArtifactMetadataFile" Type="staged-build" Platform="$(EditorPlatform)" Configuration="Development" ArtifactDir="$(ProjectOutputDirectory)/$(EditorPlatformPath)CookedEditor" />
-						<Tag Files="$(ProjectOutputDirectory)/$(EditorPlatformPath)CookedEditor/..." With="#Staged CookedEditor Files"/>
-
-						<CreateArtifact Name="$(ProjectName)-Staged-Cooked-Editor-$(EditorPlatform)" Type="staged-build" BaseDir="$(ProjectOutputDirectory)/$(EditorPlatformPath)CookedEditor" Keys="Project=$(ProjectName)" Metadata="Stream=$(Branch);Platform=$(EditorPlatform);Configuration=Development;BuildName=$(BuildNamePath);dashboard-category=Development Builds" Files="#Staged CookedEditor Files"/>
-					</Case>
-					<Default>
-						<Log Message="Copying to $(NetworkOutputDirectory)/$(EditorPlatformPath)CookedEditor"/>
-						<Do If="!$(SkipPublish)">
-							<Copy Files="..." From="$(ProjectOutputDirectory)/$(EditorPlatformPath)CookedEditor" To="$(NetworkOutputDirectory)/$(EditorPlatformPath)CookedEditor"/>
-						</Do>
-					</Default>
-				</Switch>
+				<Property Name="TargetDirectory" Value="$(ProjectOutputDirectory)/$(EditorPlatformPath)CookedEditor"/>
+				<Do If="$(UseNetworkShare)">
+					<Log Message="Copying to $(NetworkOutputDirectory)/$(EditorPlatformPath)CookedEditor"/>
+					<Do If="!$(SkipPublish)">
+						<Copy Files="..." From="$(TargetDirectory)" To="$(NetworkOutputDirectory)/$(EditorPlatformPath)CookedEditor"/>
+					</Do>
+				</Do>
+				<Expand Name="WriteArtifactMetadataFile" Type="staged-build" Platform="$(EditorPlatform)" Configuration="Development" ArtifactDir="$(TargetDirectory)" />
+				<Tag Files="$(TargetDirectory)/..." With="#Staged CookedEditor Files"/>
+				<Do If="$(Horde)">
+					<CreateArtifact Name="$(ProjectName)-Staged-Cooked-Editor-$(EditorPlatform)" Type="staged-build" BaseDir="$(TargetDirectory)" Keys="Project=$(ProjectName)" Metadata="Stream=$(Branch);Platform=$(EditorPlatform);Configuration=Development;BuildName=$(BuildNamePath);dashboard-category=Development Builds" Files="#Staged CookedEditor Files"/>
+				</Do>
+				<Do If="$(PublishCloudArtifacts)">
+					<CreateCloudArtifact Name="Cloud-$(ProjectName)-Staged-Cooked-Editor-$(EditorPlatform)" Type="staged-build" BaseDir="$(TargetDirectory)" Metadata="Configuration=Development;dashboard-category=Staged Builds - Cloud" Project="$(TargetProject)" Platform="$(EditorPlatform)" Files="#Staged CookedEditor Files"/>
+				</Do>
+				<Delete Directories="$(TargetDirectory)" Files="artifactmetadata.json"/>
 			</Node>
 
 			<!-- Aggregate Nodes requirements -->
 			<Property Name="GeneratedCookedEditorNodes" Value="$(PreNodeName)Compile CookedEditor $(EditorPlatform)" If="!$(SkipBuild)" />
 			<Property Name="GeneratedCookedEditorNodes" Value="$(GeneratedCookedEditorNodes);$(PreNodeName)Cook CookedEditor $(EditorPlatform)" If="!$(SkipCook)" />
-			<Property Name="GeneratedCookedEditorNodes" Value="$(GeneratedCookedEditorNodes);$(PlatformPublishNodeName)" If="!$(SkipPublish) And $(PublishBuilds)" />
+			<Property Name="GeneratedCookedEditorNodes" Value="$(GeneratedCookedEditorNodes);$(PlatformPublishNodeName)" If="!$(SkipPublish) And ($(PublishBuilds) Or $(Horde) Or $(PublishCloudArtifacts))" />
 		</Agent>
 
 		<!-- Tests -->
@@ -452,7 +455,7 @@
 			<Property Name="CookedEditorBuildPath" Value="$(NetworkOutputDirectory)/$(EditorPlatformPath)CookedEditor" />
 			<!-- Prerequisites for the editor test to run -->
 			<Property Name="TestPrerequisites" Value="$(GeneratedCookedEditorNodes)"/>
-			<Do If="$(Horde)">
+			<Do If="$(Horde) And !$(RetrieveUsingCloudArtifacts)">
 				<Property Name="HordeArtifactPlatformRetrieve" Value="$(PreNodeName)Retrieve CookedEditor Build Artifact $(EditorPlatform) (Horde)"/>
 				<Node Name="$(HordeArtifactPlatformRetrieve)" Requires="$(TestPrerequisites);$(PlatformPublishNodeName)">
 					<RetrieveArtifact Name="$(ProjectName)-Staged-Cooked-Editor-$(EditorPlatform)" Type="staged-build" OutputDir="$(ProjectOutputDirectory)/$(EditorPlatformPath)CookedEditor" />
@@ -463,7 +466,7 @@
 			<Do If="$(RetrieveUsingCloudArtifacts)">
 				<Property Name="CloudArtifactPlatformRetrieve" Value="$(PreNodeName)Retrieve CookedEditor Build Artifact $(EditorPlatform) (Cloud)"/>
 				<Node Name="$(CloudArtifactPlatformRetrieve)" Requires="$(TestPrerequisites);$(PlatformPublishNodeName)">
-					<RetrieveCloudArtifact Name="$(BuildNamePath)" Type="staged-build" OutputDir="$(ProjectOutputDirectory)/$(EditorPlatformPath)CookedEditor" Project="$(TargetProject)" Platform="$(EditorPlatform)" />
+					<RetrieveCloudArtifact Name="Cloud-$(ProjectName)-Staged-Cooked-Editor-$(EditorPlatform)" Type="staged-build" OutputDir="$(ProjectOutputDirectory)/$(EditorPlatformPath)CookedEditor" Project="$(TargetProject)" Platform="$(EditorPlatform)" />
 				</Node>
 				<Property Name="TestPrerequisites" Value="$(TestPrerequisites);$(CloudArtifactPlatformRetrieve)" />
 				<Property Name="CookedEditorBuildPath" Value="$(ProjectOutputDirectory)/$(EditorPlatformPath)CookedEditor" />
@@ -524,14 +527,6 @@
 		<Label Category="Editors" Name="$(PreLabelName)CookedEditor $(EditorPlatform)" Requires="$(GeneratedCookedEditorNodes)" />
 	</ForEach>
 
-	<ForEach Name="TargetPlatform" Values="$(TargetPlatforms)" Separator="+">
-		<!-- We cannot have an inner loop or these won't be added to the parent scope -->
-		<Property Name="TestPrerequisites$(TargetPlatform)Shipping" Value="" CreateInParentScope="true"/>
-		<Property Name="TestPrerequisites$(TargetPlatform)Test" Value="" CreateInParentScope="true" />
-		<Property Name="TestPrerequisites$(TargetPlatform)Development" Value="" CreateInParentScope="true" />
-		<Property Name="TestPrerequisites$(TargetPlatform)Debug" Value="" CreateInParentScope="true" />
-	</ForEach>
-
 	<!-- Now, for each platform we build, cook, stage, package, publish, and test -->
 	<ForEach Name="TargetPlatform" Values="$(TargetPlatforms)" Separator="+">
 
@@ -595,11 +590,11 @@
 		<Property Name="CloudArtifactPlatformPackagedNodeName" Value="$(PlatformPublishPackagedNodeName) (Cloud)" />
 		<Property Name="CloudArtifactPlatformRetrieve" Value="$(PreNodeNameWithPlatform)Retrieve Build Artifact $(TargetPlatform)" />
 		<Property Name="RequirePlatformStaged" Value="$(PlatformStageNodeName)" />
-		<Property Name="RequirePlatformStaged" Value="$(RequirePlatformStaged);$(PlatformStagedArtifactNodeName)" If="$(Horde)"/>
+		<Property Name="RequirePlatformStaged" Value="$(RequirePlatformStaged);$(PlatformStagedArtifactNodeName)" If="$(Horde) And !$(RetrieveUsingCloudArtifacts)"/>
 		<Property Name="RequirePlatformStaged" Value="$(RequirePlatformStaged);$(CloudPlatformStagedArtifactNodeName)" If="$(PublishCloudArtifacts)"/>
 		<Property Name="RequirePlatformStaged" Value="$(RequirePlatformStaged);$(PlatformPublishStagedNodeName)" If="!$(SkipPublish) And $(PublishBuilds)"/>
 		<Property Name="RequirePlatformPackaged" Value="$(PlatformPackageNodeName)" />
-		<Property Name="RequirePlatformPackaged" Value="$(RequirePlatformPackaged);$(HordeArtifactPlatformPackagedNodeName)" If="$(Horde)"/>
+		<Property Name="RequirePlatformPackaged" Value="$(RequirePlatformPackaged);$(HordeArtifactPlatformPackagedNodeName)" If="$(Horde) And !$(RetrieveUsingCloudArtifacts)"/>
 		<Property Name="RequirePlatformPackaged" Value="$(RequirePlatformPackaged);$(CloudArtifactPlatformPackagedNodeName)" If="$(PublishCloudArtifacts)"/>
 		<Property Name="RequirePlatformPackaged" Value="$(RequirePlatformPackaged);$(PlatformPublishPackagedNodeName)"  If="!$(SkipPublish) And $(PublishBuilds)"/>
 		
@@ -761,18 +756,19 @@
 			<!-- Publish (Staged) (Cloud) -->
 			<Do If="$(PublishCloudArtifacts)">
 				<Node Name="$(CloudPlatformStagedArtifactNodeName)" Requires="$(PlatformStageNodeName)">
-					<ForEach Name="PackageConfiguration" Values="$(PackageConfigurations)" Separator="+">
-						<Tag Files="$(ProjectOutputDirectory)/$(StagedPlatformFolder)/..." With="#$(PlatformStageNodeName) Files"/>
+					<Property Name="TargetDirectory" Value="$(ProjectOutputDirectory)/$(StagedPlatformFolder)"/>
+					<Expand Name="WriteArtifactMetadataFile" Type="staged-build" Platform="$(StagedPlatformFolder)" Configuration="$(PackageConfigurations)" ArtifactDir="$(TargetDirectory)" />
+					<Tag Files="$(TargetDirectory)/..." With="#$(PlatformStageNodeName) Files"/>
 
-						<CreateCloudArtifact Name="$(BuildNamePath)" HordeArtifactName="Cloud-$(ProjectName)-Staged-$(StagedPlatformFolder)-$(PackageConfiguration)" Type="staged-build" BaseDir="$(ProjectOutputDirectory)/$(StagedPlatformFolder)" Metadata="Configuration=$(PackageConfiguration);dashboard-category=Staged Builds - Cloud" Project="$(TargetProject)" Platform="$(StagedPlatformFolder)" Files="#$(PlatformStageNodeName) Files"/>
-					</ForEach>
+					<CreateCloudArtifact Name="Cloud-$(ProjectName)-Staged-$(StagedPlatformFolder)" Type="staged-build" BaseDir="$(TargetDirectory)" Metadata="Configuration=$(PackageConfigurations);dashboard-category=Staged Builds - Cloud" Project="$(TargetProject)" Platform="$(StagedPlatformFolder)" Files="#$(PlatformStageNodeName) Files"/>
+					<Delete Directories="$(TargetDirectory)" Files="artifactmetadata.json"/>
 				</Node>
 			</Do>
 
 			<!-- Publish (Staged) (Horde) -->
 			<Do If="$(Horde)">
 				<Node Name="$(PlatformStagedArtifactNodeName)" Requires="$(PlatformStageNodeName)">
-					<Expand Name="WriteArtifactMetadataFile" Type="staged-build" Platform="$(StagedPlatformFolder)" Configuration="$(PackageConfigurations)" ArtifactDir="$(ProjectOutputDirectory)/$(StagedPlatformFolder)" />
+					<Expand Name="WriteArtifactMetadataFile" Type="staged-build" Platform="$(StagedPlatformFolder)" Configuration="$(PackageConfigurations)" ArtifactDir="$(ProjectOutputDirectory)/$(StagedPlatformFolder)"/>
 					<Tag Files="$(ProjectOutputDirectory)/$(StagedPlatformFolder)/..." With="#$(PlatformStageNodeName) Files"/>
 
 					<CreateArtifact Name="$(ProjectName)-Staged-$(StagedPlatformFolder)" Type="staged-build" BaseDir="$(ProjectOutputDirectory)/$(StagedPlatformFolder)" Keys="Project=$(ProjectName)" Metadata="Stream=$(Branch);Platform=$(StagedPlatformFolder);BuildName=$(BuildNamePath);dashboard-category=Staged Builds" Files="#$(PlatformStageNodeName) Files"/>
@@ -793,6 +789,24 @@
 					<Command Name="BuildCookRun" Arguments="$(BCRArgs) -configuration=$(PackageConfiguration) $(StageAndPackageArguments) -skipbuild -skipcook -skipstage -skippak -skippackage -archive -archivedirectory=&quot;$(NetworkOutputDirectory)/$(StagedPlatformFolder)/Packages&quot;" />
 				</ForEach>
 			</Node>
+
+			<!-- Publish (Packages) (Cloud) -->
+			<Do  If="$(PublishCloudArtifacts)">
+				<Property Name="CloudPackageSkipBCRArguments" Value="-stage -package"/>
+				<Property Name="CloudPackageSkipBCRArguments" Value="-skipstage -skippackage" If="'$(TargetPlatform)' == 'Mac'"/>
+
+				<Node Name="$(CloudArtifactPlatformPackagedNodeName)" Requires="$(PlatformPackageNodeName)">
+					<ForEach Name="PackageConfiguration" Values="$(PackageConfigurations)" Separator="+">
+						<Property Name="TargetDirectory" Value="$(ProjectOutputDirectory)/Packages/$(StagedPlatformFolder)/$(PackageConfiguration)"/>
+						<Command Name="BuildCookRun" Arguments="$(BCRArgs) -configuration=$(PackageConfiguration) $(StageAndPackageArguments) -skipbuild -skipcook -skippak $(CloudPackageSkipBCRArguments) -archive -archivedirectory=&quot;$(TargetDirectory)&quot;" />
+						<Expand Name="WriteArtifactMetadataFile" Type="packaged-build" Platform="$(StagedPlatformFolder)" Configuration="$(PackageConfiguration)" ArtifactDir="$(TargetDirectory)" />
+						<Tag Files="$(TargetDirectory)/..." With="#$(PlatformPackageNodeName) $(PackageConfiguration) Files"/>
+
+						<CreateCloudArtifact Name="Cloud-$(ProjectName)-Packaged-$(StagedPlatformFolder)-$(PackageConfiguration)" Type="packaged-build" BaseDir="$(TargetDirectory)" Metadata="Configuration=$(PackageConfiguration);dashboard-category=Packaged Builds - Cloud" Project="$(TargetProject)" Platform="$(StagedPlatformFolder)" Files="#$(PlatformPackageNodeName) $(PackageConfiguration) Files"/>
+						<Delete Directories="$(TargetDirectory)" Files="artifactmetadata.json"/>
+					</ForEach>
+				</Node>
+			</Do>
 			
 			<!-- Publish (Packages) (Horde) -->
 			<Do  If="$(Horde)">
@@ -802,7 +816,7 @@
 				<Node Name="$(HordeArtifactPlatformPackagedNodeName)" Requires="$(PlatformPackageNodeName)">
 					<ForEach Name="PackageConfiguration" Values="$(PackageConfigurations)" Separator="+">
 						<Command Name="BuildCookRun" Arguments="$(BCRArgs) -configuration=$(PackageConfiguration) $(StageAndPackageArguments) -skipbuild -skipcook -skippak $(HordePackageSkipBCRArguments) -archive -archivedirectory=&quot;$(ProjectOutputDirectory)/Packages/$(StagedPlatformFolder)/$(PackageConfiguration)&quot;" />
-						<Expand Name="WriteArtifactMetadataFile" Type="packaged-build" Platform="$(StagedPlatformFolder)" Configuration="$(PackageConfiguration)" ArtifactDir="$(ProjectOutputDirectory)/Packages/$(StagedPlatformFolder)/$(PackageConfiguration)" />
+						<Expand Name="WriteArtifactMetadataFile" Type="packaged-build" Platform="$(StagedPlatformFolder)" Configuration="$(PackageConfiguration)" ArtifactDir="$(ProjectOutputDirectory)/Packages/$(StagedPlatformFolder)/$(PackageConfiguration)"/>
 						<Tag Files="$(ProjectOutputDirectory)/Packages/$(StagedPlatformFolder)/$(PackageConfiguration)/..." With="#$(PlatformPackageNodeName) $(PackageConfiguration) Files"/>
 
 						<CreateArtifact Name="$(ProjectName)-Packaged-$(StagedPlatformFolder)-$(PackageConfiguration)" Type="packaged-build" BaseDir="$(ProjectOutputDirectory)/Packages/$(StagedPlatformFolder)/$(PackageConfiguration)" Keys="Project=$(ProjectName)" Metadata="Stream=$(Branch);Platform=$(StagedPlatformFolder);Configuration=$(PackageConfiguration);BuildName=$(BuildNamePath);dashboard-category=$(PackageConfiguration) Builds" Files="#$(PlatformPackageNodeName) $(PackageConfiguration) Files"/>
@@ -810,21 +824,6 @@
 				</Node>
 			</Do>
 
-			<!-- Publish (Packages) (Cloud) -->
-			<Do  If="$(PublishCloudArtifacts)">
-				<Property Name="CloudPackageSkipBCRArguments" Value="-stage -package"/>
-				<Property Name="CloudPackageSkipBCRArguments" Value="-skipstage -skippackage" If="'$(TargetPlatform)' == 'Mac'"/>
-
-				<Node Name="$(CloudArtifactPlatformPackagedNodeName)" Requires="$(PlatformPackageNodeName)">
-					<ForEach Name="PackageConfiguration" Values="$(PackageConfigurations)" Separator="+">
-						<Command Name="BuildCookRun" Arguments="$(BCRArgs) -configuration=$(PackageConfiguration) $(StageAndPackageArguments) -skipbuild -skipcook -skippak $(CloudPackageSkipBCRArguments) -archive -archivedirectory=&quot;$(ProjectOutputDirectory)/Packages/$(StagedPlatformFolder)/$(PackageConfiguration)&quot;" />
-						<Tag Files="$(ProjectOutputDirectory)/Packages/$(StagedPlatformFolder)/$(PackageConfiguration)/..." With="#$(PlatformPackageNodeName) $(PackageConfiguration) Files"/>
-
-						<CreateCloudArtifact Name="$(BuildNamePath)" HordeArtifactName="Cloud-$(ProjectName)-Packaged-$(StagedPlatformFolder)-$(PackageConfiguration)" Type="packaged-build" BaseDir="$(ProjectOutputDirectory)/Packages/$(StagedPlatformFolder)/$(PackageConfiguration)" Metadata="Configuration=$(PackageConfiguration);dashboard-category=Packaged Builds - Cloud" Project="$(TargetProject)" Platform="$(StagedPlatformFolder)" Files="#$(PlatformPackageNodeName) $(PackageConfiguration) Files"/>
-					</ForEach>
-				</Node>
-			</Do>
-
 			<!-- Check Build Size -->
 			<Node Name="$(PlatformCheckBuildSizeNodeName)" Requires="$(PlatformPackageNodeName)">
 				<ForEach Name="PackageConfiguration" Values="$(PackageConfigurations)" Separator="+">
@@ -853,130 +852,147 @@
 			<Property Name="GauntletBuildArg" Value="$(GauntletBuildArg) -destlocalinstalldir=&quot;$(LocalInstallDirectory)&quot;" If="!$(IsLocalPlatform) And $(WithDeviceReservation)"/>
 		</Do>
 
-		<Property Name="GauntletBuildArg$(TargetPlatform)" Value="$(GauntletBuildArg)" CreateInParentScope="true"/>
+		<!-- Test Configs -->
+		<ForEach Name="TestConfig" Values="$(TestConfigList)" Separator="^">
+			<Property Name="TestConfigTitle" Value=""/>
+			<Regex Input="$(TestConfig)" Pattern="(.*)\|(.*)\|(.*)\|(.*)\|(.*)\|(.*)" Capture="TestConfigName;TestConfigTargetPlatform;ConfigArtifactPostfix;ConfigAdditionalTestArguments;ConfigAdditionalAppArguments;ConfigAgentType"/>
 
-		<Agent Name="Test $(PlatformAgentName)" Type="$(TestAgentType)">
-			<!-- Run through all configs -->
-			<ForEach Name="TestConfiguration" Values="$(TestConfigurations)" Separator="+">
-				<Property Name="NodeAnnotations" Value=""/>
-				<!-- Prerequisites for the test to run -->
-				<Property Name="TestPrerequisites" Value="$(RequirePlatformStaged)" />
-				<Property Name="TestPrerequisites" Value="$(RequirePlatformPackaged)" If="$(UsePackaged)"/>
-				<Property Name="TestPrerequisites" Value="$(PlatformToolsNodeName)" If="'$(GeneratedBuildVersion)' != ''"/>
-				<Property Name="TestPrerequisites" Value="$(PlatformToolsNodeName);$(PlatformEditorNodeName)" If="'$(GeneratedBuildVersion)' != '' and $(TestsRequireEditor)"/>
-
-				<Do If="$(Horde) and '$(RetrieveUsingCloudArtifacts)' == 'false'">
-					<Property Name="HordeArtifactPlatformRetrieveWithConfig" Value="$(HordeArtifactPlatformRetrieve) $(TestConfiguration) (Horde)"/>
-					<Node Name="$(HordeArtifactPlatformRetrieveWithConfig)" Requires="$(TestPrerequisites)">
-						<Switch>
-							<Case If="$(UsePackaged)">
-								<RetrieveArtifact Name="$(ProjectName)-Packaged-$(StagedPlatformFolder)-$(TestConfiguration)" Type="packaged-build" OutputDir="$(ProjectOutputDirectory)/Packages/$(StagedPlatformFolder)/$(TestConfiguration)" />
-							</Case>
-							<Default>
-								<RetrieveArtifact Name="$(ProjectName)-Staged-$(StagedPlatformFolder)" Type="staged-build" OutputDir="$(ProjectOutputDirectory)/$(StagedPlatformFolder)" />
-							</Default>
-						</Switch>
-					</Node>
-					<Property Name="TestPrerequisites" Value="$(TestPrerequisites);$(HordeArtifactPlatformRetrieveWithConfig)" />
-				</Do>
+			<!--If this is the default test config i.e. "All" or Empty ensure we create nodes for it-->
+			<Property Name="TestConfigTargetPlatform" Value="$(TargetPlatform)" If="'$(TestConfigTargetPlatform)' == 'All' Or '$(TestConfigTargetPlatform)' == ''"/>
 
-				<Do If="'$(RetrieveUsingCloudArtifacts)' == 'true'">
-					<Property Name="CloudArtifactPlatformRetrieveWithConfig" Value="$(CloudArtifactPlatformRetrieve) $(TestConfiguration) (Cloud)"/>
-					<Node Name="$(CloudArtifactPlatformRetrieveWithConfig)" Requires="$(TestPrerequisites)">
-						<Switch>
-							<Case If="$(UsePackaged)">
-								<RetrieveCloudArtifact Name="$(BuildNamePath)" Type="packaged-build" OutputDir="$(ProjectOutputDirectory)/Packages/$(StagedPlatformFolder)/$(TestConfiguration)" Project="$(TargetProject)" Platform="$(StagedPlatformFolder)" />
-							</Case>
-							<Default>
-								<RetrieveCloudArtifact Name="$(BuildNamePath)" Type="staged-build" OutputDir="$(ProjectOutputDirectory)/$(StagedPlatformFolder)" Project="$(TargetProject)" Platform="$(StagedPlatformFolder)" />
-							</Default>
-						</Switch>
-					</Node>
-
-					<Property Name="TestPrerequisites" Value="$(TestPrerequisites);$(CloudArtifactPlatformRetrieveWithConfig)" />
-				</Do>
+			<Do If="'$(TestConfigTargetPlatform)' == '$(TargetPlatform)'">
 
-				<Property Name="GauntletTestArgs" Value="$(GauntletTargetTestArgs)"/>
-				<Property Name="GauntletTestArgs" Value="$(GauntletTargetTestAltArgs)" If="ContainsItem('$(AlternateDevicePoolPlatforms)', '$(TargetPlatform)', ';')"/>
-
-				<Property Name="GauntletTestArgs" Value="$(GauntletTestArgs) -adbkeys=&quot;$(AdbKeysPath)&quot;" If="'$(AdbKeysPath)' != '' and '$(TargetPlatform)' == 'Android'"/>
-
-				<!-- Install step -->
-				<Do If="!$(IsLocalPlatform) And $(WithDeviceReservation)">
-					<Property Name="NodeAnnotations" Value=""/>
-					<Property Name="NodeAnnotations" Value="Workflow=$(AnnotationsTarget)" If="ContainsItem('$(AnnotationsPlatforms)', '$(TargetPlatform)', ';') and ContainsItem('$(AnnotationsTests)', 'Install', ';')"/>
-					<Property Name="InstallNodeName" Value="$(PreNodeName)$(TargetPlatform) $(TestConfiguration) Install Game"/>
-					<Property Name="InstallArgs" Value="-ClientCount=1 $(GauntletPreInstallArgs)"/>
-					<Do If="'$(AutomationReportOutputDirectory)' != ''">
-						<Property Name="OutputFragmentPath" Value="$(TargetPlatform)/$(TestConfiguration)/UE.InstallOnly" />
-						<Property Name="OutputReportPath" Value="$(AutomationReportOutputDirectory)/$(OutputFragmentPath)" />
-						<Property Name="GauntletReportArgs" Value="-logdir=&quot;$(OutputReportPath)&quot; -artifactname=Gauntlet" />
-					</Do>
-					<Node Name="$(InstallNodeName)" Requires="$(TestPrerequisites)" Annotations="DeviceReserve=Begin;$(NodeAnnotations)">
-						<Command Name="RunUnreal" Arguments="-project=$(TargetProject) -test=&quot;UE.InstallOnly&quot; -platform=$(TargetPlatform) -configuration=$(TestConfiguration) -target=$(TargetName) $(GauntletBuildArg) $(GauntletTestArgs) $(GauntletReportArgs) $(InstallArgs)" />
-					</Node>
-					<Property Name="TestPrerequisites" Value="$(TestPrerequisites);$(InstallNodeName)" />
-					<Do If="!ContainsItem('$(TargetPlatformsToSkipBootTests)', '$(TargetPlatform)', '+') Or !ContainsItem('$(TargetPlatformsToSkipTests)', '$(TargetPlatform)', '+')">
-						<Property Name="AllPlatformInstallNodes" Value="$(AllPlatformInstallNodes);$(InstallNodeName)"/>
-					</Do>
-					<!-- Skip install for all other test node -->
-					<Property Name="GauntletTestArgs" Value="$(GauntletTestArgs) $(GauntletTestSkipInstallArgs)"/>
-				</Do>
+				<Property Name="TestConfigTitle" Value=" $(TestConfigName)" If="'$(TestConfigName)' != ''"/>
 
-				<!-- Declare a boot test for the target-->
-				<Property Name="NodeAnnotations" Value=""/>
-				<Property Name="NodeAnnotations" Value="Workflow=$(AnnotationsTarget)" If="ContainsItem('$(AnnotationsPlatforms)', '$(TargetPlatform)', ';') and ContainsItem('$(AnnotationsTests)', 'BootTest', ';')"/>
-				<Property Name="BootTestName" Value="$(PreNodeName)$(TargetPlatform) $(TestConfiguration) BootTest" />
-				<Node Name="$(BootTestName)" Requires="$(TestPrerequisites)" Annotations="$(NodeAnnotations)">
-					<Do If="'$(AutomationReportOutputDirectory)' != ''">
-						<Property Name="OutputFragmentPath" Value="$(TargetPlatform)/$(TestConfiguration)/UE.BootTest" />
-						<Property Name="OutputReportPath" Value="$(AutomationReportOutputDirectory)/$(OutputFragmentPath)" />
-						<Property Name="GauntletReportArgs" Value="-logdir=&quot;$(OutputReportPath)&quot; -artifactname=Gauntlet" />
-					</Do>
-					<!-- Run the boot test -->
-					<Command Name="RunUnreal" Arguments="-project=$(TargetProject) -test=&quot;UE.BootTest&quot; -platform=$(TargetPlatform) -configuration=$(TestConfiguration) -target=$(TargetName) $(GauntletBuildArg) $(GauntletTestArgs) $(GauntletReportArgs)" />
-				</Node>
-				<!-- Link the boot test in to the prerequisites and generated nodes if its enabled-->
-				<Property Name="TestPrerequisites" Value="$(TestPrerequisites);$(PlatformEditorNodeName)" />
-				<Do If="$(TargetBootTest) and !ContainsItem('$(TargetPlatformsToSkipBootTests)', '$(TargetPlatform)', '+')">
-					<Property Name="TestPrerequisites" Value="$(TestPrerequisites);$(BootTestName)" />
-					<Property Name="AllPlatformBootTestNodes" Value="$(AllPlatformBootTestNodes);$(BootTestName)"/>
-				</Do>
+				<Property Name="TestConfigArgs" Value="" />
+				<Property Name="TestConfigArgs" Value="$(ConfigAdditionalTestArguments)" If="'$(ConfigAdditionalTestArguments)' != ''"/>
+				<Property Name="TestConfigArgs" Value="$(TestConfigArgs) -ArtifactPostfix=$(ConfigArtifactPostfix)" If="'$(ConfigArtifactPostfix)' != ''"  />
+				<Property Name="TestConfigArgs" Value="$(TestConfigArgs) -args=&quot;$(ConfigAdditionalAppArguments)&quot;" If="'$(ConfigAdditionalAppArguments)' != ''" />
+
+				<Property Name="TestTargetAgentType" Value="$(TestAgentType)"/>
+				<Property Name="TestTargetAgentType" Value="$(ConfigAgentType)" If="'$(ConfigAgentType)' != ''"/>
 
-				<Property Name="TestPrerequisites$(TargetPlatform)$(TestConfiguration)" Value="$(TestPrerequisites)" />
-
-				<Do If="!$(UseTestConfigs)">
-					<!-- Go through the list of tests -->	
-					<ForEach Name="TestName" Values="$(TargetTestList)" Separator="+">
-						<Property Name="TestNameTitle" Value="$(TestName)"/>
-						<Expand Name="SanitizeTestNameTitle"/>
-						<!-- Run and report on the test-->
-						<Do If="'$(AutomationReportOutputDirectory)' != ''">
-							<Property Name="OutputFragmentPath" Value="$(TargetPlatform)/$(TestConfiguration)/$(TestNameTitle)" />
-							<Property Name="OutputReportPath" Value="$(AutomationReportOutputDirectory)/$(OutputFragmentPath)" />
-							<Property Name="GauntletReportArgs" Value="-logdir=&quot;$(OutputReportPath)&quot; -artifactname=Gauntlet -ReportExportPath=&quot;$(OutputReportPath)&quot;" />
-							<Property Name="GauntletReportArgs" Value="$(GauntletReportArgs) -ReportURL=&quot;$(AutomationReportUrl)/$(OutputFragmentPath)/&quot;" If="'$(AutomationReportUrl)' != ''" />
+				<Agent Name="Test$(TestConfigTitle) $(PlatformAgentName)" Type="$(TestTargetAgentType)">
+					<!-- Run through all configs -->
+					<ForEach Name="TestConfiguration" Values="$(TestConfigurations)" Separator="+">
+						<Property Name="NodeAnnotations" Value=""/>
+						<!-- Prerequisites for the test to run -->
+						<Property Name="TestPrerequisites" Value="$(RequirePlatformStaged)" />
+						<Property Name="TestPrerequisites" Value="$(RequirePlatformPackaged)" If="$(UsePackaged)"/>
+						<Property Name="TestPrerequisites" Value="$(PlatformToolsNodeName)" If="'$(GeneratedBuildVersion)' != ''"/>
+						<Property Name="TestPrerequisites" Value="$(PlatformToolsNodeName);$(PlatformEditorNodeName)" If="'$(GeneratedBuildVersion)' != '' and $(TestsRequireEditor)"/>
+
+						<Do If="$(Horde) and !$(RetrieveUsingCloudArtifacts)">
+							<Property Name="HordeArtifactPlatformRetrieveWithConfig" Value="$(HordeArtifactPlatformRetrieve) $(TestConfiguration)$(TestConfigTitle) (Horde)"/>
+							<Node Name="$(HordeArtifactPlatformRetrieveWithConfig)" Requires="$(TestPrerequisites)" IgnoreModified="artifactmetadata.json">
+								<Switch>
+									<Case If="$(UsePackaged)">
+										<RetrieveArtifact Name="$(ProjectName)-Packaged-$(StagedPlatformFolder)-$(TestConfiguration)" Type="packaged-build" OutputDir="$(ProjectOutputDirectory)/Packages/$(StagedPlatformFolder)/$(TestConfiguration)" />
+									</Case>
+									<Default>
+										<RetrieveArtifact Name="$(ProjectName)-Staged-$(StagedPlatformFolder)" Type="staged-build" OutputDir="$(ProjectOutputDirectory)/$(StagedPlatformFolder)" />
+									</Default>
+								</Switch>
+							</Node>
+							<Property Name="TestPrerequisites" Value="$(TestPrerequisites);$(HordeArtifactPlatformRetrieveWithConfig)" />
+						</Do>
+
+						<Do If="$(RetrieveUsingCloudArtifacts)">
+							<Property Name="CloudArtifactPlatformRetrieveWithConfig" Value="$(CloudArtifactPlatformRetrieve) $(TestConfiguration)$(TestConfigTitle) (Cloud)"/>
+							<Node Name="$(CloudArtifactPlatformRetrieveWithConfig)" Requires="$(TestPrerequisites)" IgnoreModified="artifactmetadata.json">
+								<Switch>
+									<Case If="$(UsePackaged)">
+										<RetrieveCloudArtifact Name="Cloud-$(ProjectName)-Packaged-$(StagedPlatformFolder)-$(TestConfiguration)" Type="packaged-build" OutputDir="$(ProjectOutputDirectory)/Packages/$(StagedPlatformFolder)/$(TestConfiguration)" Project="$(TargetProject)" Platform="$(StagedPlatformFolder)" />
+									</Case>
+									<Default>
+										<RetrieveCloudArtifact Name="Cloud-$(ProjectName)-Staged-$(StagedPlatformFolder)" Type="staged-build" OutputDir="$(ProjectOutputDirectory)/$(StagedPlatformFolder)" Project="$(TargetProject)" Platform="$(StagedPlatformFolder)" />
+									</Default>
+								</Switch>
+							</Node>
+
+							<Property Name="TestPrerequisites" Value="$(TestPrerequisites);$(CloudArtifactPlatformRetrieveWithConfig)" />
 						</Do>
-						<Property Name="TestNodeName" Value="$(PreNodeName)$(TargetPlatform) $(TestConfiguration) $(TestNameTitle)" />
+
+						<Property Name="GauntletTestArgs" Value="$(GauntletTargetTestArgs)"/>
+						<Property Name="GauntletTestArgs" Value="$(GauntletTargetTestAltArgs)" If="ContainsItem('$(AlternateDevicePoolPlatforms)', '$(TargetPlatform)', ';')"/>
+
+						<Property Name="GauntletTestArgs" Value="$(GauntletTestArgs) -adbkeys=&quot;$(AdbKeysPath)&quot;" If="'$(AdbKeysPath)' != '' and '$(TargetPlatform)' == 'Android'"/>
+
+						<!-- Install step -->
+						<Do If="!$(IsLocalPlatform) And $(WithDeviceReservation)">
+							<Property Name="NodeAnnotations" Value=""/>
+							<Property Name="NodeAnnotations" Value="Workflow=$(AnnotationsTarget)" If="ContainsItem('$(AnnotationsPlatforms)', '$(TargetPlatform)', ';') and ContainsItem('$(AnnotationsTests)', 'Install', ';')"/>
+							<Property Name="InstallNodeName" Value="$(PreNodeName)$(TargetPlatform) $(TestConfiguration)$(TestConfigTitle) Install Game"/>
+							<Property Name="InstallArgs" Value="-ClientCount=1 $(GauntletPreInstallArgs)"/>
+							<Do If="'$(AutomationReportOutputDirectory)' != ''">
+								<Property Name="OutputFragmentPath" Value="$(TargetPlatform)/$(TestConfiguration)/UE.InstallOnly" />
+								<Property Name="OutputReportPath" Value="$(AutomationReportOutputDirectory)/$(OutputFragmentPath)" />
+								<Property Name="GauntletReportArgs" Value="-logdir=&quot;$(OutputReportPath)&quot; -artifactname=Gauntlet" />
+							</Do>
+							<Node Name="$(InstallNodeName)" Requires="$(TestPrerequisites)" Annotations="DeviceReserve=Begin;$(NodeAnnotations)">
+								<Command Name="RunUnreal" Arguments="-project=$(TargetProject) -test=&quot;UE.InstallOnly&quot; -platform=$(TargetPlatform) -configuration=$(TestConfiguration) -target=$(TargetName) $(GauntletBuildArg) $(GauntletTestArgs) $(GauntletReportArgs) $(InstallArgs)" />
+							</Node>
+							<Property Name="TestPrerequisites" Value="$(TestPrerequisites);$(InstallNodeName)" />
+							<Do If="!ContainsItem('$(TargetPlatformsToSkipBootTests)', '$(TargetPlatform)', '+') Or !ContainsItem('$(TargetPlatformsToSkipTests)', '$(TargetPlatform)', '+')">
+								<Property Name="AllPlatformInstallNodes" Value="$(AllPlatformInstallNodes);$(InstallNodeName)"/>
+							</Do>
+							<!-- Skip install for all other test node -->
+							<Property Name="GauntletTestArgs" Value="$(GauntletTestArgs) $(GauntletTestSkipInstallArgs)"/>
+						</Do>
+
+						<!-- Declare a boot test for the target-->
 						<Property Name="NodeAnnotations" Value=""/>
-						<Property Name="NodeAnnotations" Value="Workflow=$(AnnotationsTarget)" If="ContainsItem('$(AnnotationsPlatforms)', '$(TargetPlatform)', ';') and ContainsItem('$(AnnotationsTests)', '$(TestName)', ';')"/>
-						<!-- Run the test. No 'Requires' as we run on the same machine locally and use the publish folder otherwise -->
-						<Property Name="TestNameBeforeParentheses" Value="$(TestName)"/>
-						<Property Name="TestArgumentsWithinParentheses" Value=""/>
-						<Expand Name="SplitCombinedTestNameAndArgs"/>
-						<Node Name="$(TestNodeName)" Requires="$(TestPrerequisites)" Annotations="$(NodeAnnotations)">
-							<Command Name="RunUnreal" Arguments="-project=$(TargetProject) -test=&quot;$(TestNameBeforeParentheses)&quot; -platform=$(TargetPlatform) -configuration=$(TestConfiguration) -target=$(TargetName) $(GauntletBuildArg) $(GauntletTestArgs) $(TestArgumentsWithinParentheses) $(GauntletReportArgs)"/>
+						<Property Name="NodeAnnotations" Value="Workflow=$(AnnotationsTarget)" If="ContainsItem('$(AnnotationsPlatforms)', '$(TargetPlatform)', ';') and ContainsItem('$(AnnotationsTests)', 'BootTest', ';')"/>
+						<Property Name="BootTestName" Value="$(PreNodeName)$(TargetPlatform) $(TestConfiguration)$(TestConfigTitle) BootTest" />
+						<Node Name="$(BootTestName)" Requires="$(TestPrerequisites)" Annotations="$(NodeAnnotations)">
+							<Do If="'$(AutomationReportOutputDirectory)' != ''">
+								<Property Name="OutputFragmentPath" Value="$(TargetPlatform)/$(TestConfiguration)/UE.BootTest" />
+								<Property Name="OutputReportPath" Value="$(AutomationReportOutputDirectory)/$(OutputFragmentPath)" />
+								<Property Name="GauntletReportArgs" Value="-logdir=&quot;$(OutputReportPath)&quot; -artifactname=Gauntlet" />
+							</Do>
+							<!-- Run the boot test -->
+							<Command Name="RunUnreal" Arguments="-project=$(TargetProject) -test=&quot;UE.BootTest&quot; -platform=$(TargetPlatform) -configuration=$(TestConfiguration) -target=$(TargetName) $(GauntletBuildArg) $(GauntletTestArgs) $(GauntletReportArgs) $(TestConfigArgs)" />
 						</Node>
-						<!-- Add it to our list -->
-						<Do If="!ContainsItem('$(TargetPlatformsToSkipTests)', '$(TargetPlatform)', '+')">
-							<Property Name="AllPlatformTestNodes" Value="$(AllPlatformTestNodes);$(TestNodeName)"/>
+						<!-- Link the boot test in to the prerequisites and generated nodes if its enabled-->
+						<Property Name="TestPrerequisites" Value="$(TestPrerequisites);$(PlatformEditorNodeName)" />
+						<Do If="$(TargetBootTest) and !ContainsItem('$(TargetPlatformsToSkipBootTests)', '$(TargetPlatform)', '+')">
+							<Property Name="TestPrerequisites" Value="$(TestPrerequisites);$(BootTestName)" />
+							<Property Name="AllPlatformBootTestNodes" Value="$(AllPlatformBootTestNodes);$(BootTestName)"/>
 						</Do>
+
+						<!-- Go through the list of tests -->
+						<ForEach Name="TestName" Values="$(TargetTestList)" Separator="+">
+							<Property Name="TestNameTitle" Value="$(TestName)"/>
+							<Expand Name="SanitizeTestNameTitle"/>
+							<!-- Run and report on the test-->
+							<Do If="'$(AutomationReportOutputDirectory)' != ''">
+								<Property Name="OutputFragmentPath" Value="$(TargetPlatform)/$(TestConfiguration)/$(TestNameTitle)" />
+								<Property Name="OutputReportPath" Value="$(AutomationReportOutputDirectory)/$(OutputFragmentPath)" />
+								<Property Name="GauntletReportArgs" Value="-logdir=&quot;$(OutputReportPath)&quot; -artifactname=Gauntlet -ReportExportPath=&quot;$(OutputReportPath)&quot;" />
+								<Property Name="GauntletReportArgs" Value="$(GauntletReportArgs) -ReportURL=&quot;$(AutomationReportUrl)/$(OutputFragmentPath)/&quot;" If="'$(AutomationReportUrl)' != ''" />
+							</Do>
+							<Property Name="TestNodeName" Value="$(PreNodeName)$(TargetPlatform) $(TestConfiguration)$(TestConfigTitle) $(TestNameTitle)" />
+							<Property Name="NodeAnnotations" Value=""/>
+							<Property Name="NodeAnnotations" Value="Workflow=$(AnnotationsTarget)" If="ContainsItem('$(AnnotationsPlatforms)', '$(TargetPlatform)', ';') and ContainsItem('$(AnnotationsTests)', '$(TestName)', ';')"/>
+							<!-- Run the test. No 'Requires' as we run on the same machine locally and use the publish folder otherwise -->
+							<Property Name="TestNameBeforeParentheses" Value="$(TestName)"/>
+							<Property Name="TestArgumentsWithinParentheses" Value=""/>
+							<Expand Name="SplitCombinedTestNameAndArgs"/>
+							<Node Name="$(TestNodeName)" Requires="$(TestPrerequisites)" Annotations="$(NodeAnnotations)">
+								<Command Name="RunUnreal" Arguments="-project=$(TargetProject) -test=&quot;$(TestNameBeforeParentheses)&quot; -platform=$(TargetPlatform) -configuration=$(TestConfiguration) -target=$(TargetName) $(GauntletBuildArg) $(GauntletTestArgs) $(TestArgumentsWithinParentheses) $(GauntletReportArgs) $(TestConfigArgs)"/>
+							</Node>
+							<!-- Add it to our list -->
+							<Do If="!ContainsItem('$(TargetPlatformsToSkipTests)', '$(TargetPlatform)', '+')">
+								<Property Name="AllPlatformTestNodes" Value="$(AllPlatformTestNodes);$(TestNodeName)"/>
+							</Do>
+						</ForEach>
+
+						<!-- Expand the tests here -->
+						<Expand Name="AdditionalTestNodes"/>
 					</ForEach>
-				</Do>
-				<!-- Expand the tests here -->
-				<Expand Name="AdditionalTestNodes"/>
-			</ForEach>
-		</Agent>
+				</Agent>
+			</Do>
+		</ForEach>
 		
 		<!-- Declare labels for CIS -->
 		<Label Category="Clients" Name="$(PreLabelName)$(TargetPlatform)" Requires="$(PlatformCompileNodeName);$(UploadSymbolsNodeName)" />
@@ -998,21 +1014,8 @@
 
 		<Property Name="GeneratedTargetNodes" Value="$(GeneratedTargetNodes);$(PlatformTasks)"/>
 		<Property Name="GeneratedTargetTestNodes" Value="$(GeneratedTargetTestNodes);$(AllPlatformBootTestNodes);$(AllPlatformTestNodes)"/>
-
-		<Property Name="GeneratedTargetTestNodes$(TargetPlatform)" Value="$(GeneratedTargetTestNodes)" CreateInParentScope="true"/>
 	</ForEach>
 
-	<Do If="$(UseTestConfigs)">
-		<Include Script="RunTestConfigs.xml"/>
-		<!-- If you have test configs set and you don't have any, include the default one. If this has issues. we can stop it doing -->
-		<Do If="'$(AllTestConfigs)' == ''">
-			<Include Script="TestConfigs/DefaultTestConfig.xml"/>
-		</Do>
-		<ForEach Name="TargetPlatform" Values="$(TargetPlatforms)" Separator="+">
-			<Property Name="GeneratedTargetTestNodes" Value="$(GeneratedTargetTestNodes);$(GeneratedTargetTestNodes$(TargetPlatform))"/>
-		</ForEach>
-	</Do>
-
 	<!-- Create an aggregate with all the generated nodes -->
 	<Aggregate Name="BuildAndTest $(ProjectName)" Requires="$(GeneratedToolNodes);$(GeneratedCookedEditorNodes);$(GeneratedEditorTestNodes);$(GeneratedCookedEditorTestNodes);$(GeneratedTargetNodes);$(GeneratedTargetTestNodes)"/>
 	<Aggregate Name="OnlyTest $(ProjectName)" Requires="$(GeneratedEditorTestNodes);$(GeneratedCookedEditorTestNodes);$(GeneratedTargetTestNodes)"/>
diff --git a/Engine/Build/Graph/Tasks/Inc/TestConfigs.xml b/Engine/Build/Graph/Tasks/Inc/TestConfigs.xml
new file mode 100644
--- /dev/null
+++ b/Engine/Build/Graph/Tasks/Inc/TestConfigs.xml
@@ -0,0 +1,24 @@
+<?xml version='1.0' ?>
+<BuildGraph xmlns="http://www.epicgames.com/BuildGraph" xmlns:xsi="http://www.w3.org/2001/XMLSchema-instance" xsi:schemaLocation="http://www.epicgames.com/BuildGraph ../../../../Engine/Build/Graph/Schema.xsd" >
+
+	<!-- Test Configs allows to override the default homogenous configurations passed to tests -->
+	<!-- This paradigm allows to explicitly declare test configurations for each target set -->
+
+	<!-- Include this file with: <Include Script="$(RootDir)/Engine/Build/Graph/Tasks/Inc/TestConfigs.xml"/> -->
+	<!-- Then Expand the Macro 'AddTestConfig' with corresponding configuration inputs. ie: -->
+	<!-- <Expand Name="AddTestConfig"
+			TestConfigName="Vulkan SM6"
+			TestPlatform="Win64"
+			AdditionalAppArguments="-vulkan -sm6"
+			ArtifactPostfix="Vulkan_SM6"
+		/> -->
+
+	<Option Name="TestConfigList" DefaultValue="" Description="Input TestConfigs; expand macro AddTestConfig after including Engine/Build/Graph/Tasks/Inc/TestConfigs.xml to populate"/>
+
+	<Do If="'$(TestConfigList)' == ''">
+		<Macro Name="AddTestConfig" Arguments="TestConfigName" OptionalArguments="TestPlatform;AdditionalTestArguments;ArtifactPostfix;AdditionalAppArguments;AgentType">
+			<Property Name="TestConfigList" Value="$(TestConfigList)^$(TestConfigName)|$(TestPlatform)|$(ArtifactPostfix)|$(AdditionalTestArguments)|$(AdditionalAppArguments)|$(AgentType)"/>
+		</Macro>
+	</Do>
+
+</BuildGraph>
\ No newline at end of file
diff --git a/Engine/Build/Graph/Tasks/RunTestConfigs.xml b/Engine/Build/Graph/Tasks/RunTestConfigs.xml
deleted file mode 100644
diff --git a/Engine/Build/Graph/Tasks/TestConfigs.xml b/Engine/Build/Graph/Tasks/TestConfigs.xml
deleted file mode 100644
diff --git a/Engine/Build/Graph/Tasks/TestConfigs/DefaultTestConfig.xml b/Engine/Build/Graph/Tasks/TestConfigs/DefaultTestConfig.xml
deleted file mode 100644
diff --git a/Engine/Plugins/Performance/AutomatedPerfTesting/Build/Inc/TestConfigs.xml b/Engine/Plugins/Performance/AutomatedPerfTesting/Build/Inc/TestConfigs.xml
--- a/Engine/Plugins/Performance/AutomatedPerfTesting/Build/Inc/TestConfigs.xml
+++ b/Engine/Plugins/Performance/AutomatedPerfTesting/Build/Inc/TestConfigs.xml
@@ -1,10 +1,12 @@
 <?xml version='1.0' ?>
-<BuildGraph xmlns="http://www.epicgames.com/BuildGraph" xmlns:xsi="http://www.w3.org/2001/XMLSchema-instance" xsi:schemaLocation="http://www.epicgames.com/BuildGraph ../../../../Engine/Build/Graph/Schema.xsd" >
+<BuildGraph xmlns="http://www.epicgames.com/BuildGraph" xmlns:xsi="http://www.w3.org/2001/XMLSchema-instance" xsi:schemaLocation="http://www.epicgames.com/BuildGraph ../../../../../Build/Graph/Schema.xsd" >
 
-	<Property Name="AllTestConfigs" Value=""/>
-	
-	<Macro Name="AddTestConfig" Arguments="TestPlatform;TestConfigName;ArtifactPostfix;AdditionalArguments;AgentType">
-		<Property Name="AllTestConfigs" Value="$(AllTestConfigs);$(TestPlatform)|$(TestConfigName)|$(ArtifactPostfix)|$(AdditionalArguments)"/>
+	<Include Script="../../../../../Build/Graph/Tasks/Inc/TestConfigs.xml"/>
+
+	<Macro Name="AddPerfTestConfig" Arguments="TestConfigName;TestPlatform" OptionalArguments="AgentType;AppArguments">
+		<Property Name="AdditionalTestArguments" Value="-AutomatedPerfTest.TestConfigName=&quot;$(TestConfigName)&quot;"/>
+		<StringOp Input="$(TestConfigName)" Method="Replace" Arguments=" ;_" Output="ArtifactPostfix"/>
+		<Property Name="TestConfigList" Value="$(TestConfigList)^$(TestConfigName)|$(TestPlatform)|$(ArtifactPostfix)|$(AdditionalTestArguments)|$(AppArguments)|$(AgentType)"/>
 	</Macro>
 
 </BuildGraph>
