Revert enabling RetrieveUsingCloudArtifact on Epic Buildmachine

#rnx

diff --git a/Engine/Build/Graph/Tasks/BuildAndTestProject.xml b/Engine/Build/Graph/Tasks/BuildAndTestProject.xml
--- a/Engine/Build/Graph/Tasks/BuildAndTestProject.xml
+++ b/Engine/Build/Graph/Tasks/BuildAndTestProject.xml
@@ -209,12 +209,11 @@
 	<!-- On Epic build machines, we want to:
 		1. Write dependency information into the cooked output so that it is suitable for use as a base upon which to cook incrementally.
 		2. Cook incrementally if the initial state at the start of the job is not empty/blank.
-		3. Use Cloud artifacts by default to publish and retrieve build
 	-->
 	<Do If="$(IsEpicBuildMachine)">
 		<Property Name="ExtraCookArguments" Value="$(ExtraCookArguments) -cookincremental"/>
 	</Do>
-	<Option Name="RetrieveUsingCloudArtifacts" DefaultValue="$(IsEpicBuildMachine)" Restrict="true|false" Description="Enable to retrieve cloud artifacts" />
+	<Option Name="RetrieveUsingCloudArtifacts" DefaultValue="false" Restrict="true|false" Description="Enable to retrieve cloud artifacts" />
 	<Option Name="PublishCloudArtifacts" DefaultValue="$(RetrieveUsingCloudArtifacts)" Description="Enable to publish cloud artifacts" />
 
 	<!-- Detect the type of the client based on the target name -->
