Fix Delete node attribute input + Remove unnecessary IgnoreModified file

#rnx

diff --git a/Engine/Build/Graph/Tasks/BuildAndTestProject.xml b/Engine/Build/Graph/Tasks/BuildAndTestProject.xml
--- a/Engine/Build/Graph/Tasks/BuildAndTestProject.xml
+++ b/Engine/Build/Graph/Tasks/BuildAndTestProject.xml
@@ -439,7 +439,7 @@
 				<Do If="$(PublishCloudArtifacts)">
 					<CreateCloudArtifact Name="Cloud-$(ProjectName)-Staged-Cooked-Editor-$(EditorPlatform)" Type="staged-build" BaseDir="$(TargetDirectory)" Metadata="Configuration=Development;dashboard-category=Staged Builds - Cloud" Project="$(TargetProject)" Platform="$(EditorPlatform)" Files="#Staged CookedEditor Files"/>
 				</Do>
-				<Delete Directories="$(TargetDirectory)" Files="artifactmetadata.json"/>
+				<Delete Files="$(TargetDirectory)/artifactmetadata.json"/>
 			</Node>
 
 			<!-- Aggregate Nodes requirements -->
@@ -760,7 +760,7 @@
 					<Tag Files="$(TargetDirectory)/..." With="#$(PlatformStageNodeName) Files"/>
 
 					<CreateCloudArtifact Name="Cloud-$(ProjectName)-Staged-$(StagedPlatformFolder)" Type="staged-build" BaseDir="$(TargetDirectory)" Metadata="Configuration=$(PackageConfigurations);dashboard-category=Staged Builds - Cloud" Project="$(TargetProject)" Platform="$(StagedPlatformFolder)" Files="#$(PlatformStageNodeName) Files"/>
-					<Delete Directories="$(TargetDirectory)" Files="artifactmetadata.json"/>
+					<Delete Files="$(TargetDirectory)/artifactmetadata.json"/>
 				</Node>
 			</Do>
 
@@ -802,7 +802,7 @@
 						<Tag Files="$(TargetDirectory)/..." With="#$(PlatformPackageNodeName) $(PackageConfiguration) Files"/>
 
 						<CreateCloudArtifact Name="Cloud-$(ProjectName)-Packaged-$(StagedPlatformFolder)-$(PackageConfiguration)" Type="packaged-build" BaseDir="$(TargetDirectory)" Metadata="Configuration=$(PackageConfiguration);dashboard-category=Packaged Builds - Cloud" Project="$(TargetProject)" Platform="$(StagedPlatformFolder)" Files="#$(PlatformPackageNodeName) $(PackageConfiguration) Files"/>
-						<Delete Directories="$(TargetDirectory)" Files="artifactmetadata.json"/>
+						<Delete Files="$(TargetDirectory)/artifactmetadata.json"/>
 					</ForEach>
 				</Node>
 			</Do>
@@ -883,7 +883,7 @@
 
 						<Do If="$(Horde) and !$(RetrieveUsingCloudArtifacts)">
 							<Property Name="HordeArtifactPlatformRetrieveWithConfig" Value="$(HordeArtifactPlatformRetrieve) $(TestConfiguration)$(TestConfigTitle) (Horde)"/>
-							<Node Name="$(HordeArtifactPlatformRetrieveWithConfig)" Requires="$(TestPrerequisites)" IgnoreModified="artifactmetadata.json">
+							<Node Name="$(HordeArtifactPlatformRetrieveWithConfig)" Requires="$(TestPrerequisites)">
 								<Switch>
 									<Case If="$(UsePackaged)">
 										<RetrieveArtifact Name="$(ProjectName)-Packaged-$(StagedPlatformFolder)-$(TestConfiguration)" Type="packaged-build" OutputDir="$(ProjectOutputDirectory)/Packages/$(StagedPlatformFolder)/$(TestConfiguration)" />
@@ -898,7 +898,7 @@
 
 						<Do If="$(RetrieveUsingCloudArtifacts)">
 							<Property Name="CloudArtifactPlatformRetrieveWithConfig" Value="$(CloudArtifactPlatformRetrieve) $(TestConfiguration)$(TestConfigTitle) (Cloud)"/>
-							<Node Name="$(CloudArtifactPlatformRetrieveWithConfig)" Requires="$(TestPrerequisites)" IgnoreModified="artifactmetadata.json">
+							<Node Name="$(CloudArtifactPlatformRetrieveWithConfig)" Requires="$(TestPrerequisites)">
 								<Switch>
 									<Case If="$(UsePackaged)">
 										<RetrieveCloudArtifact Name="Cloud-$(ProjectName)-Packaged-$(StagedPlatformFolder)-$(TestConfiguration)" Type="packaged-build" OutputDir="$(ProjectOutputDirectory)/Packages/$(StagedPlatformFolder)/$(TestConfiguration)" Project="$(TargetProject)" Platform="$(StagedPlatformFolder)" />
