Change naming convention for Cloud Artifact storage

#rnx

diff --git a/Engine/Build/Graph/Tasks/BuildAndTestProject.xml b/Engine/Build/Graph/Tasks/BuildAndTestProject.xml
--- a/Engine/Build/Graph/Tasks/BuildAndTestProject.xml
+++ b/Engine/Build/Graph/Tasks/BuildAndTestProject.xml
@@ -445,7 +445,7 @@
 					<CreateArtifact Name="$(ProjectName)-Staged-Cooked-Editor-$(EditorPlatform)" Type="staged-build" BaseDir="$(TargetDirectory)" Keys="Project=$(ProjectName)" Metadata="Stream=$(Branch);Platform=$(EditorPlatform);Configuration=Development;BuildName=$(BuildNamePath);dashboard-category=Development Builds" Files="#Staged CookedEditor Files"/>
 				</Do>
 				<Do If="$(PublishCloudArtifacts)">
-					<CreateCloudArtifact Name="Cloud-$(ProjectName)-Staged-Cooked-Editor-$(EditorPlatform)" Type="staged-build" BaseDir="$(TargetDirectory)" Metadata="Configuration=Development;dashboard-category=Staged Builds - Cloud" Project="$(TargetProject)" Platform="$(EditorPlatform)" Files="#Staged CookedEditor Files"/>
+					<CreateCloudArtifact Name="$(BuildNamePath)-Staged-Cooked-Editor" HordeArtifactName="Cloud-$(ProjectName)-Staged-Cooked-Editor-$(EditorPlatform)" Type="staged-build" BaseDir="$(TargetDirectory)" Metadata="Configuration=Development;dashboard-category=Staged Builds - Cloud" Project="$(TargetProject)" Platform="$(EditorPlatform)" Files="#Staged CookedEditor Files"/>
 				</Do>
 				<Delete Files="$(TargetDirectory)/artifactmetadata.json"/>
 			</Node>
@@ -473,7 +473,7 @@
 			<Do If="$(RetrieveUsingCloudArtifacts)">
 				<Property Name="CloudArtifactPlatformRetrieve" Value="$(PreNodeName)Retrieve CookedEditor Build Artifact $(EditorPlatform) (Cloud)"/>
 				<Node Name="$(CloudArtifactPlatformRetrieve)" Requires="$(TestPrerequisites);$(PlatformPublishNodeName)">
-					<RetrieveCloudArtifact Name="Cloud-$(ProjectName)-Staged-Cooked-Editor-$(EditorPlatform)" Type="staged-build" OutputDir="$(ProjectOutputDirectory)/$(EditorPlatformPath)CookedEditor" Project="$(TargetProject)" Platform="$(EditorPlatform)" />
+					<RetrieveCloudArtifact Name="$(BuildNamePath)-Staged-Cooked-Editor" Type="staged-build" OutputDir="$(ProjectOutputDirectory)/$(EditorPlatformPath)CookedEditor" Project="$(TargetProject)" Platform="$(EditorPlatform)" />
 				</Node>
 				<Property Name="TestPrerequisites" Value="$(TestPrerequisites);$(CloudArtifactPlatformRetrieve)" />
 				<Property Name="CookedEditorBuildPath" Value="$(ProjectOutputDirectory)/$(EditorPlatformPath)CookedEditor" />
@@ -767,7 +767,7 @@
 					<Expand Name="WriteArtifactMetadataFile" Type="staged-build" Platform="$(StagedPlatformFolder)" Configuration="$(PackageConfigurations)" ArtifactDir="$(TargetDirectory)" />
 					<Tag Files="$(TargetDirectory)/..." With="#$(PlatformStageNodeName) Files"/>
 
-					<CreateCloudArtifact Name="Cloud-$(ProjectName)-Staged-$(StagedPlatformFolder)" Type="staged-build" BaseDir="$(TargetDirectory)" Metadata="Configuration=$(PackageConfigurations);dashboard-category=Staged Builds - Cloud" Project="$(TargetProject)" Platform="$(StagedPlatformFolder)" Files="#$(PlatformStageNodeName) Files"/>
+					<CreateCloudArtifact Name="$(BuildNamePath)-Staged" HordeArtifactName="Cloud-$(ProjectName)-Staged-$(StagedPlatformFolder)" Type="staged-build" BaseDir="$(TargetDirectory)" Metadata="dashboard-category=Staged Builds - Cloud" Project="$(TargetProject)" Platform="$(StagedPlatformFolder)" Files="#$(PlatformStageNodeName) Files"/>
 					<Delete Files="$(TargetDirectory)/artifactmetadata.json"/>
 				</Node>
 			</Do>
@@ -809,7 +809,7 @@
 						<Expand Name="WriteArtifactMetadataFile" Type="packaged-build" Platform="$(StagedPlatformFolder)" Configuration="$(PackageConfiguration)" ArtifactDir="$(TargetDirectory)" />
 						<Tag Files="$(TargetDirectory)/..." With="#$(PlatformPackageNodeName) $(PackageConfiguration) Files"/>
 
-						<CreateCloudArtifact Name="Cloud-$(ProjectName)-Packaged-$(StagedPlatformFolder)-$(PackageConfiguration)" Type="packaged-build" BaseDir="$(TargetDirectory)" Metadata="Configuration=$(PackageConfiguration);dashboard-category=Packaged Builds - Cloud" Project="$(TargetProject)" Platform="$(StagedPlatformFolder)" Files="#$(PlatformPackageNodeName) $(PackageConfiguration) Files"/>
+						<CreateCloudArtifact Name="$(BuildNamePath)-Packaged-$(PackageConfiguration)" HordeArtifactName="Cloud-$(ProjectName)-Packaged-$(StagedPlatformFolder)-$(PackageConfiguration)" Type="packaged-build" BaseDir="$(TargetDirectory)" Metadata="Configuration=$(PackageConfiguration);dashboard-category=Packaged Builds - Cloud" Project="$(TargetProject)" Platform="$(StagedPlatformFolder)" Files="#$(PlatformPackageNodeName) $(PackageConfiguration) Files"/>
 						<Delete Files="$(TargetDirectory)/artifactmetadata.json"/>
 					</ForEach>
 				</Node>
@@ -909,10 +909,10 @@
 							<Node Name="$(CloudArtifactPlatformRetrieveWithConfig)" Requires="$(TestPrerequisites)">
 								<Switch>
 									<Case If="$(UsePackaged)">
-										<RetrieveCloudArtifact Name="Cloud-$(ProjectName)-Packaged-$(StagedPlatformFolder)-$(TestConfiguration)" Type="packaged-build" OutputDir="$(ProjectOutputDirectory)/Packages/$(StagedPlatformFolder)/$(TestConfiguration)" Project="$(TargetProject)" Platform="$(StagedPlatformFolder)" />
+										<RetrieveCloudArtifact Name="$(BuildNamePath)-Packaged-$(TestConfiguration)" Type="packaged-build" OutputDir="$(ProjectOutputDirectory)/Packages/$(StagedPlatformFolder)/$(TestConfiguration)" Project="$(TargetProject)" Platform="$(StagedPlatformFolder)" />
 									</Case>
 									<Default>
-										<RetrieveCloudArtifact Name="Cloud-$(ProjectName)-Staged-$(StagedPlatformFolder)" Type="staged-build" OutputDir="$(ProjectOutputDirectory)/$(StagedPlatformFolder)" Project="$(TargetProject)" Platform="$(StagedPlatformFolder)" />
+										<RetrieveCloudArtifact Name="$(BuildNamePath)-Staged" Type="staged-build" OutputDir="$(ProjectOutputDirectory)/$(StagedPlatformFolder)" Project="$(TargetProject)" Platform="$(StagedPlatformFolder)" />
 									</Default>
 								</Switch>
 							</Node>
