Allow publishing Horde artifact even if Cloud Artifact is enabled

#rnx

diff --git a/Engine/Build/Graph/Tasks/BuildAndTestProject.xml b/Engine/Build/Graph/Tasks/BuildAndTestProject.xml
--- a/Engine/Build/Graph/Tasks/BuildAndTestProject.xml
+++ b/Engine/Build/Graph/Tasks/BuildAndTestProject.xml
@@ -216,6 +216,7 @@
 	</Do>
 	<Option Name="RetrieveUsingCloudArtifacts" DefaultValue="false" Restrict="true|false" Description="Enable to retrieve cloud artifacts" />
 	<Option Name="PublishCloudArtifacts" DefaultValue="$(RetrieveUsingCloudArtifacts)" Description="Enable to publish cloud artifacts" />
+	<Option Name="PublishHordeArtifacts" DefaultValue="$(IsEpicBuildMachine)" Description="Enable to publish horde artifacts"/>
 
 	<!-- If running with GPU Reshape, override the local build configuration
 	     Enables shader symbols and anything else that aids troubleshooting -->
@@ -441,7 +442,7 @@
 				</Do>
 				<Expand Name="WriteArtifactMetadataFile" Type="staged-build" Platform="$(EditorPlatform)" Configuration="Development" ArtifactDir="$(TargetDirectory)" />
 				<Tag Files="$(TargetDirectory)/..." With="#Staged CookedEditor Files"/>
-				<Do If="$(Horde)">
+				<Do If="$(PublishHordeArtifacts)">
 					<CreateArtifact Name="$(ProjectName)-Staged-Cooked-Editor-$(EditorPlatform)" Type="staged-build" BaseDir="$(TargetDirectory)" Keys="Project=$(ProjectName)" Metadata="Stream=$(Branch);Platform=$(EditorPlatform);Configuration=Development;BuildName=$(BuildNamePath);dashboard-category=Development Builds" Files="#Staged CookedEditor Files"/>
 				</Do>
 				<Do If="$(PublishCloudArtifacts)">
@@ -453,7 +454,7 @@
 			<!-- Aggregate Nodes requirements -->
 			<Property Name="GeneratedCookedEditorNodes" Value="$(PreNodeName)Compile CookedEditor $(EditorPlatform)" If="!$(SkipBuild)" />
 			<Property Name="GeneratedCookedEditorNodes" Value="$(GeneratedCookedEditorNodes);$(PreNodeName)Cook CookedEditor $(EditorPlatform)" If="!$(SkipCook)" />
-			<Property Name="GeneratedCookedEditorNodes" Value="$(GeneratedCookedEditorNodes);$(PlatformPublishNodeName)" If="!$(SkipPublish) And ($(PublishBuilds) Or $(Horde) Or $(PublishCloudArtifacts))" />
+			<Property Name="GeneratedCookedEditorNodes" Value="$(GeneratedCookedEditorNodes);$(PlatformPublishNodeName)" If="!$(SkipPublish) And ($(PublishBuilds) Or $(PublishHordeArtifacts) Or $(PublishCloudArtifacts))" />
 		</Agent>
 
 		<!-- Tests -->
@@ -462,7 +463,7 @@
 			<Property Name="CookedEditorBuildPath" Value="$(NetworkOutputDirectory)/$(EditorPlatformPath)CookedEditor" />
 			<!-- Prerequisites for the editor test to run -->
 			<Property Name="TestPrerequisites" Value="$(GeneratedCookedEditorNodes)"/>
-			<Do If="$(Horde) And !$(RetrieveUsingCloudArtifacts)">
+			<Do If="$(PublishHordeArtifacts) And !$(RetrieveUsingCloudArtifacts)">
 				<Property Name="HordeArtifactPlatformRetrieve" Value="$(PreNodeName)Retrieve CookedEditor Build Artifact $(EditorPlatform) (Horde)"/>
 				<Node Name="$(HordeArtifactPlatformRetrieve)" Requires="$(TestPrerequisites);$(PlatformPublishNodeName)">
 					<RetrieveArtifact Name="$(ProjectName)-Staged-Cooked-Editor-$(EditorPlatform)" Type="staged-build" OutputDir="$(ProjectOutputDirectory)/$(EditorPlatformPath)CookedEditor" />
@@ -597,12 +598,12 @@
 		<Property Name="CloudArtifactPlatformPackagedNodeName" Value="$(PlatformPublishPackagedNodeName) (Cloud)" />
 		<Property Name="CloudArtifactPlatformRetrieve" Value="$(PreNodeNameWithPlatform)Retrieve Build Artifact $(TargetPlatform)" />
 		<Property Name="RequirePlatformStaged" Value="$(PlatformStageNodeName)" />
-		<Property Name="RequirePlatformStaged" Value="$(RequirePlatformStaged);$(PlatformStagedArtifactNodeName)" If="$(Horde) And !$(RetrieveUsingCloudArtifacts)"/>
-		<Property Name="RequirePlatformStaged" Value="$(RequirePlatformStaged);$(CloudPlatformStagedArtifactNodeName)" If="$(PublishCloudArtifacts)"/>
+		<Property Name="RequirePlatformStaged" Value="$(RequirePlatformStaged);$(PlatformStagedArtifactNodeName)" If="!$(SkipPublish) And $(PublishHordeArtifacts)"/>
+		<Property Name="RequirePlatformStaged" Value="$(RequirePlatformStaged);$(CloudPlatformStagedArtifactNodeName)" If="!$(SkipPublish) And $(PublishCloudArtifacts)"/>
 		<Property Name="RequirePlatformStaged" Value="$(RequirePlatformStaged);$(PlatformPublishStagedNodeName)" If="!$(SkipPublish) And $(PublishBuilds)"/>
 		<Property Name="RequirePlatformPackaged" Value="$(PlatformPackageNodeName)" />
-		<Property Name="RequirePlatformPackaged" Value="$(RequirePlatformPackaged);$(HordeArtifactPlatformPackagedNodeName)" If="$(Horde) And !$(RetrieveUsingCloudArtifacts)"/>
-		<Property Name="RequirePlatformPackaged" Value="$(RequirePlatformPackaged);$(CloudArtifactPlatformPackagedNodeName)" If="$(PublishCloudArtifacts)"/>
+		<Property Name="RequirePlatformPackaged" Value="$(RequirePlatformPackaged);$(HordeArtifactPlatformPackagedNodeName)" If="!$(SkipPublish) And $(PublishHordeArtifacts)"/>
+		<Property Name="RequirePlatformPackaged" Value="$(RequirePlatformPackaged);$(CloudArtifactPlatformPackagedNodeName)" If="!$(SkipPublish) And $(PublishCloudArtifacts)"/>
 		<Property Name="RequirePlatformPackaged" Value="$(RequirePlatformPackaged);$(PlatformPublishPackagedNodeName)"  If="!$(SkipPublish) And $(PublishBuilds)"/>
 		
 		<!-- Set CookPlatform since it isn't always the same as the target -->
@@ -773,7 +774,7 @@
 			</Do>
 
 			<!-- Publish (Staged) (Horde) -->
-			<Do If="$(Horde)">
+			<Do If="$(PublishHordeArtifacts)">
 				<Node Name="$(PlatformStagedArtifactNodeName)" Requires="$(PlatformStageNodeName)">
 					<Expand Name="WriteArtifactMetadataFile" Type="staged-build" Platform="$(StagedPlatformFolder)" Configuration="$(PackageConfigurations)" ArtifactDir="$(ProjectOutputDirectory)/$(StagedPlatformFolder)"/>
 					<Tag Files="$(ProjectOutputDirectory)/$(StagedPlatformFolder)/..." With="#$(PlatformStageNodeName) Files"/>
@@ -816,7 +817,7 @@
 			</Do>
 			
 			<!-- Publish (Packages) (Horde) -->
-			<Do  If="$(Horde)">
+			<Do  If="$(PublishHordeArtifacts)">
 				<Property Name="HordePackageSkipBCRArguments" Value="-stage -package"/>
 				<Property Name="HordePackageSkipBCRArguments" Value="-skipstage -skippackage" If="'$(TargetPlatform)' == 'Mac'"/>
 
@@ -889,7 +890,7 @@
 						<Property Name="TestPrerequisites" Value="$(PlatformToolsNodeName)" If="'$(GeneratedBuildVersion)' != ''"/>
 						<Property Name="TestPrerequisites" Value="$(PlatformToolsNodeName);$(PlatformEditorNodeName)" If="'$(GeneratedBuildVersion)' != '' and $(TestsRequireEditor)"/>
 
-						<Do If="$(Horde) and !$(RetrieveUsingCloudArtifacts)">
+						<Do If="$(PublishHordeArtifacts) and !$(RetrieveUsingCloudArtifacts)">
 							<Property Name="HordeArtifactPlatformRetrieveWithConfig" Value="$(HordeArtifactPlatformRetrieve) $(TestConfiguration)$(TestConfigTitle) (Horde)"/>
 							<Node Name="$(HordeArtifactPlatformRetrieveWithConfig)" Requires="$(TestPrerequisites)">
 								<Switch>
