NaniteTest: Add initial version of TestConfigs

diff --git a/Engine/Build/Graph/Tasks/BuildAndTestProject.xml b/Engine/Build/Graph/Tasks/BuildAndTestProject.xml
--- a/Engine/Build/Graph/Tasks/BuildAndTestProject.xml
+++ b/Engine/Build/Graph/Tasks/BuildAndTestProject.xml
@@ -42,6 +42,7 @@
 	<Option Name="ExtraToolCompileArguments" DefaultValue="" Restrict="" Description="Extra arguments to use when building the tools" />
 	<Option Name="ExtraEditorCompileArguments" DefaultValue="" Restrict="" Description="Extra arguments to use when building the editor" />
 	<Option Name="ExtraTargetCompileArguments" DefaultValue="" Restrict="" Description="Extra arguments to use when building targets" />
+	<Option Name="ExtraTargetEntitlementArguments" DefaultValue="" Restrict="" Description="Extra arguments to use when building targets" />
 	<Option Name="MapIniSection" DefaultValue="" Restrict="" Description="MapIniSection argument to pass when cooking" />
 	<Option Name="ExtraCookArguments" DefaultValue="" Restrict="" Description="Extra arguments to use when cooking" />
 	<Option Name="ExtraStageAndPackageArguments" DefaultValue="" Restrict="" Description="Extra arguments to use when staging and packaging" />
@@ -201,10 +202,11 @@
 		<Property Name="PlatformsToSkipPackaging" Value="$(EditorPlatforms)+$(TargetPlatforms)" />
 	</Do>
 
+	<Property Name="TargetCompileArguments" Value="$(ExtraTargetCompileArguments) $(ExtraTargetEntitlementArguments)" />
 	<Do If="$(UseIncrementalAgents)" >
 		<Property Name="ExtraToolCompileArguments" Value="$(ExtraToolCompileArguments) -allmodules -nodebuginfo -Strict -WarningsAsErrors"/>
 		<Property Name="ExtraEditorCompileArguments" Value="$(ExtraEditorCompileArguments) -allmodules -nodebuginfo -Strict -WarningsAsErrors"/>
-		<Property Name="ExtraTargetCompileArguments" Value="$(ExtraTargetCompileArguments) -nodebuginfo -WarningsAsErrors"/>
+		<Property Name="TargetCompileArguments" Value="$(TargetCompileArguments) -nodebuginfo -WarningsAsErrors"/>
 	</Do>
 
 	<!-- On Epic build machines, we want to:
@@ -642,7 +644,7 @@
 			<Node Name="$(PlatformCompileNodeName)" Requires="$(PlatformToolsNodeName)">
 				<Switch>
 					<Case If="$(IsContentOnlyProject)">
-						<Command Name="BuildGame" Arguments="$(BCRArgs) -configuration=$(TargetConfigurations) -notools -ubtargs=&quot;-BuildVersion=$(BuildVersion) $(GenericCompileArguments) $(ExtraTargetCompileArguments)&quot;" />
+						<Command Name="BuildGame" Arguments="$(BCRArgs) -configuration=$(TargetConfigurations) -notools -ubtargs=&quot;-BuildVersion=$(BuildVersion) $(GenericCompileArguments) $(TargetCompileArguments)&quot;" />
 						<Tag Files="$(ProjectPath)/Binaries/..." Filter="$(TargetName)*.target" With="#$(PlatformCompileNodeName)"/>
 						<Tag Files="$(ProjectPath)/Binaries/..." Filter="UnrealGame*.target" With="#UnrealGameTarget"/>
 						<Copy Files="#UnrealGameTarget" From="$(ProjectPath)" To="Engine" Tag="#$(PlatformCompileNodeName)"/>
@@ -650,9 +652,9 @@
 					</Case>
 					<Default>
 						<ForEach Name="TargetConfiguration" Values="$(TargetConfigurations)" Separator="+">
-							<Compile Target="$(TargetName)" Platform="$(TargetPlatform)" Configuration="$(TargetConfiguration)" Arguments="-BuildVersion=&quot;$(BuildVersion)&quot; $(GenericCompileArguments) $(ExtraTargetCompileArguments)"/>
+							<Compile Target="$(TargetName)" Platform="$(TargetPlatform)" Configuration="$(TargetConfiguration)" Arguments="-BuildVersion=&quot;$(BuildVersion)&quot; $(GenericCompileArguments) $(TargetCompileArguments)"/>
 							<Do If="'$(DedicatedServer)'" >
-								<Compile Target="$(TargetName)Server" Platform="Win64" Configuration="$(TargetConfiguration)" Arguments="-BuildVersion=&quot;$(BuildVersion)&quot; $(GenericCompileArguments) $(ExtraTargetCompileArguments)"/>
+								<Compile Target="$(TargetName)Server" Platform="Win64" Configuration="$(TargetConfiguration)" Arguments="-BuildVersion=&quot;$(BuildVersion)&quot; $(GenericCompileArguments) $(TargetCompileArguments)"/>
 							</Do>
 						</ForEach>
 					</Default>
