Add in a new extension file so you can inject a test after the BootTest but before the rest of the Gauntlet tests

diff --git a/Engine/Build/Graph/Tasks/BuildAndTestProject.xml b/Engine/Build/Graph/Tasks/BuildAndTestProject.xml
--- a/Engine/Build/Graph/Tasks/BuildAndTestProject.xml
+++ b/Engine/Build/Graph/Tasks/BuildAndTestProject.xml
@@ -125,6 +125,15 @@
 	<Option Name="TargetPlatformsToSkipBootTests" DefaultValue="$(TargetPlatformsToSkipTests)" Description="Target platforms that should skip boot tests."/>
 	<Option Name="CookedEditorPlatformsToSkipBootTests" DefaultValue="$(CookedEditorPlatformsToSkipTests)" Description="Cooked editor platforms that should skip boot tests."/>
 
+	<!-- Similar to below but they are specifically added before the other tests -->
+	<!-- -->
+	<Option Name="PreTargetTestsExtendFile" DefaultValue="" Description="For extend the AdditionalPreTestNodes Macro" />
+	<Macro Name="PreTargetTests" />
+	<Do If="'$(PreTargetTestsExtendFile)' != ''">
+		<Error Message="$(PreTargetTestsExtendFile) does not seem to exist" If="!Exists('$(PreTargetTestsExtendFile)')" />
+		<Include Script="$(PreTargetTestsExtendFile)" />
+	</Do>
+
 	<!-- Macro that allows adding additional per-platform, per-configuration test nodes -->
 	<!-- Expand the macro inside the AdditionalTestNodesExtendFile --> 
 	<Option Name="AdditionalTestNodesExtendFile" DefaultValue="" Description="For extend the AdditionalTestNodes Macro" />
@@ -970,6 +979,9 @@
 							<Property Name="AllPlatformBootTestNodes" Value="$(AllPlatformBootTestNodes);$(BootTestName)"/>
 						</Do>
 
+						<!-- Expand the tests here -->
+						<Expand Name="PreTargetTests"/>
+
 						<!-- Go through the list of tests -->
 						<ForEach Name="TestName" Values="$(TargetTestList)" Separator="+">
 							<Property Name="TestNameTitle" Value="$(TestName)"/>
