IterativeCompare horde jobs :
- Use -IgnoreBlueprintCookDependencies in the 1st cook. That's what the 2nd and 3rd cook are using.
- Add cook.displaymode=2 to the 2nd cook so I can see the list of recooked packages.

#rnx

diff --git a/Engine/Build/Graph/Tests/IterativeCompare.xml b/Engine/Build/Graph/Tests/IterativeCompare.xml
--- a/Engine/Build/Graph/Tests/IterativeCompare.xml
+++ b/Engine/Build/Graph/Tests/IterativeCompare.xml
@@ -30,7 +30,7 @@
 	<Option Name="NullCookCustomConfig" DefaultValue="" Description="What custom configuration to apply if doing a null cook (see -CustomConfig commandline argument)"/>
 	<Property Name="FullPathProjectDir" Value="$(RootDir)/$(ProjectDir)"/>
 	<Property Name="EditorPlatform" Value="Win64"/>
-	<Property Name="CommonCookArgs" Value="-targetplatform=$(TargetPlatform) -cookcultures=en -unversioned -zenstore -cookincremental"/>
+	<Property Name="CommonCookArgs" Value="-targetplatform=$(TargetPlatform) -cookcultures=en -unversioned -zenstore -cookincremental -IgnoreBlueprintCookDependencies"/>
 	<Property Name="CommonCookArgs" Value="$(CommonCookArgs) -CookProcessCount=$(CookProcessCount)" If="'$(CookProcessCount)' != ''"/>
 	<Property Name="CommonCookArgs" Value="$(CommonCookArgs) -MPCookCoreSubscription=$(MPCookCoreSubscription)" If="'$(MPCookCoreSubscription)' != ''"/>
 	<Property Name="PlatformCookedDir" Value="$(FullPathProjectDir)/Saved/Cooked/$(TargetPlatform)"/>
@@ -87,14 +87,14 @@
 			<Node Name="$(PrimingNullIncrementalCookNodeName)" Requires="$(StoreIncrementalCookNodeName)">
 				<!-- Execute the null incremental cook. -->
 				<Log Message="$(NewLine)========== IterativeCompare.xml: Priming Null Incremental Cook =========="/>
-				<Property Name="NullCookArgs" Value="$(CommonCookArgs) -IgnoreBlueprintCookDependencies"/>
+				<Property Name="NullCookArgs" Value="$(CommonCookArgs) -dpcvars=cook.displaymode=2"/>
 				<Property Name="NullCookArgs" Value="$(NullCookArgs) -CustomConfig=$(NullCookCustomConfig)" If="'$(NullCookCustomConfig)' != ''"/>
 				<Commandlet Name="Cook" Project="$(ProjectFile)" Arguments="$(NullCookArgs)"/>
 			</Node>
 			<Node Name="$(NullIncrementalCookNodeName)" Requires="$(PrimingNullIncrementalCookNodeName)">
 				<!-- Execute the null incremental cook. -->
 				<Log Message="$(NewLine)========== IterativeCompare.xml: Null Incremental Cook =========="/>
-				<Property Name="NullCookArgs" Value="$(CommonCookArgs) -IgnoreBlueprintCookDependencies -dpcvars=cook.displaymode=2"/>
+				<Property Name="NullCookArgs" Value="$(CommonCookArgs) -dpcvars=cook.displaymode=2"/>
 				<Property Name="NullCookArgs" Value="$(NullCookArgs) -CustomConfig=$(NullCookCustomConfig)" If="'$(NullCookCustomConfig)' != ''"/>
 				<Commandlet Name="Cook" Project="$(ProjectFile)" Arguments="$(NullCookArgs)"/>
 			</Node>
