Fixed #CookShaderSymbols tag redefinition

diff --git a/Engine/Build/Graph/Tasks/BuildAndTestProject.xml b/Engine/Build/Graph/Tasks/BuildAndTestProject.xml
--- a/Engine/Build/Graph/Tasks/BuildAndTestProject.xml
+++ b/Engine/Build/Graph/Tasks/BuildAndTestProject.xml
@@ -684,7 +684,7 @@
 		<Agent Name="Package $(PlatformAgentName)" Type="$(PackageAgentType)">
 
 			<!-- Cook -->
-			<Node Name="$(PlatformCookNodeName)" Requires="$(PlatformEditorNodeName)" Produces="#$(PlatformCookNodeName) Complete;#CookShaderSymbols" RunEarly="true">
+			<Node Name="$(PlatformCookNodeName)" Requires="$(PlatformEditorNodeName)" Produces="#$(PlatformCookNodeName) Complete;#$(PlatformCookNodeName) CookShaderSymbols" RunEarly="true">
 				<Do If="!$(SkipCook)">
 					<Expand Name="PreCookStep"/>
 					<Property Name="CookArgs" Value="" />
@@ -693,8 +693,8 @@
 
 					<!-- If Reshape, keep the symbols around for later -->
 					<Do If="$(DoGPUReshape)">
-						<Tag Files="$(RootDir)\$(ProjectPath)\Saved\ShaderSymbols\..." With="#CookShaderSymbols"/>
-						<Property Name="RequirePlatformStaged" Value="$(RequirePlatformStaged);#CookShaderSymbols"/>
+						<Tag Files="$(RootDir)\$(ProjectPath)\Saved\ShaderSymbols\..." With="#$(PlatformCookNodeName) CookShaderSymbols"/>
+						<Property Name="RequirePlatformStaged" Value="$(RequirePlatformStaged);#$(PlatformCookNodeName) CookShaderSymbols"/>
 					</Do>
 				</Do>
 			</Node>
