Allow IncrementalCompare job to export snapshots of its state for debug purposes.

diff --git a/Engine/Build/Graph/Tests/IncrementalCompare.xml b/Engine/Build/Graph/Tests/IncrementalCompare.xml
--- a/Engine/Build/Graph/Tests/IncrementalCompare.xml
+++ b/Engine/Build/Graph/Tests/IncrementalCompare.xml
@@ -1,6 +1,15 @@
 <?xml version='1.0' ?>
 <BuildGraph xmlns="http://www.epicgames.com/BuildGraph" xmlns:xsi="http://www.w3.org/2001/XMLSchema-instance" xsi:schemaLocation="http://www.epicgames.com/BuildGraph ../Schema.xsd" >
 
+	<EnvVar Name="UE_HORDE_JOBID"/>
+	<EnvVar Name="UE_HORDE_STEPID"/>
+	<EnvVar Name="UE_HORDE_TEMPLATEID"/>
+	
+	<EnvVar Name="UE-CloudPublishHost"/>
+	<EnvVar Name="UE-CloudPublishHttpVersion"/>
+	<EnvVar Name="UE-CloudPublishNamespace"/>
+	<EnvVar Name="UE-CloudPublishBaselineBranch"/>
+
 	<!-- Project Options -->
 	<Option Name="ProjectFile" DefaultValue="_Unknown_" Restrict=".+" Description="The projects to cook"/>
 	<Option Name="ProjectName" DefaultValue="_Unknown_" Restrict=".+" Description="The projects to cook"/>
@@ -28,6 +37,8 @@
 	<Option Name="MPCookCoreSubscription" DefaultValue="" Description="Over or Under subscription factor for CPU cores to use during cooking, default is ini-driven."/>
 	<Option Name="WithNullCook" Restrict="true|false" DefaultValue="false" Description="Whether to do a null cook or not"/>
 	<Option Name="NullCookCustomConfig" DefaultValue="" Description="What custom configuration to apply if doing a null cook (see -CustomConfig commandline argument)"/>
+	<Option Name="PublishCloudSnapshot" DefaultValue="false" Restrict="true|false" Description="Publish build to cloud host" />
+
 	<Property Name="FullPathProjectDir" Value="$(RootDir)/$(ProjectDir)"/>
 	<Property Name="EditorPlatform" Value="Win64"/>
 	<Property Name="CommonCookArgs" Value="-targetplatform=$(TargetPlatform) -cookcultures=en -unversioned -zenstore -cookincremental -IgnoreBlueprintCookDependencies"/>
@@ -62,6 +73,33 @@
 			<Commandlet Name="Cook" Project="$(ProjectFile)" Arguments="$(CommonCookArgs)"/>
 		</Node>
 		<Node Name="$(StoreIncrementalCookNodeName)" Requires="$(NaturalIncrementalCookNodeName)">
+			<!-- Exporting snapshot for developer debugging needs -->
+			<Property Name="ShouldSnapshotPublishCloud" Value="false"/>
+			<Property Name="ShouldSnapshotPublishCloud" Value="true" If="$(PublishCloudSnapshot) and '$(UE-CloudPublishHost)' != ''"/>
+
+			<Do If="$(ShouldSnapshotPublishCloud)">
+				<Property Name="IsPreflight" Value="false"/>
+				<Property Name="PreflightChange" Value=""/>
+				<Property Name="SnapshotBranch" Value="$(EscapedBranch)"/>
+				<StringOp Input="$(EscapedBranch)" Method="Replace" Arguments=".;-" Output="SnapshotBranch"/>
+				<Property Name="StepId" Value="$(UE_HORDE_STEPID)"/>
+				<Property Name="BuildVersion" Value="$(EscapedBranch)-CL-$(Change)"/>
+				<Property Name="SnapshotIdentifier" Value="$(UE_HORDE_TEMPLATEID).$(BuildVersion).$(UE_HORDE_JOBID:-6).$(StepId)"/>
+				<Property Name="SnapshotBuildBaselineBranch" Value="$(Branch)"/>
+				<Property Name="SnapshotBuildBaselineBranch" Value="$(UE-CloudPublishBaselineBranch)" If="'$(UE-CloudPublishBaselineBranch)' != ''"/>
+				<Property Name="IsHordeEnv" Value="false"/>
+				<Property Name="IsHordeEnv" Value="true" If="'$(UE_HORDE_JOBID)' != ''"/>
+				<Property Name="TargetType" Value="Server"/>
+				<Property Name="ClientPrefix" Value=""/>
+				<StringOp Input="$(TargetPlatform)" Method="SplitFirst" Arguments="Client" Output="ClientPrefix"/>
+				<Property Name="TargetType" Value="Client" If="'$(ClientPrefix)' != ''"/>
+
+				<Property Name="SnapshotBuildMetaData" Value="name=$(SnapshotIdentifier).$(TargetPlatform);depot=$(Depot);branch=$(Branch);baselineBranch=$(SnapshotBuildBaselineBranch);project=$(ProjectName);runtime=$(TargetType);platform=Win64;cookPlatform=$(TargetPlatform);changelist=$(Change);codeChangelist=$(CodeChange);isPreflight=$(IsPreflight);isBuildMachine=$(IsBuildMachine);isHordeEnv=$(IsHordeEnv)"/>
+				<Property Name="SnapshotBuildMetaData" Value="$(SnapshotBuildMetaData);preflightChangelist=$(PreflightChange)" If="'$(PreflightChange)' != ''" />
+				<Property Name="SnapshotBuildMetaData" Value="$(SnapshotBuildMetaData);hordeTemplateId=$(UE_HORDE_TEMPLATEID)" If="'$(UE_HORDE_TEMPLATEID)' != ''" />
+				<Log Message="Exporting cook state with metadata '$(SnapshotBuildMetaData)'"/>
+				<ZenExportSnapshot Project="$(ProjectFile)" Platform="$(TargetPlatform)" DestinationStorageType="Builds" DestinationCloudHost="$(UE-CloudPublishHost)" DestinationCloudHttpVersion="$(UE-CloudPublishHttpVersion)" DestinationCloudNamespace="$(UE-CloudPublishNamespace)" DestinationCloudBucket="$(ProjectName).oplog-incrementalcompare.$(SnapshotBranch).$(TargetPlatform)" DestinationIdentifier="$(SnapshotIdentifier).$(TargetPlatform)" Metadata="$(SnapshotBuildMetaData)"/>
+			</Do>
 			<!-- Move all cook artifacts to the compare_incremental dir. -->
 			<Log Message="$(NewLine)========== IncrementalCompare.xml: Move =========="/>
 			<Log Message="Delete Directories=$(CompareDirIncremental)"/>
