Cloud metadata unification part 1: bring exported snapshot metadata in line with common keys used in staged/packaged builds.

diff --git a/Engine/Build/Graph/Tasks/BuildAndTestProject.xml b/Engine/Build/Graph/Tasks/BuildAndTestProject.xml
--- a/Engine/Build/Graph/Tasks/BuildAndTestProject.xml
+++ b/Engine/Build/Graph/Tasks/BuildAndTestProject.xml
@@ -743,10 +743,12 @@
 					<Do If="$(ShouldSnapshotPublishCloud)">
 						<Property Name="SnapshotBranch" Value="$(EscapedBranch)"/>
 						<StringOp Input="$(EscapedBranch)" Method="Replace" Arguments=".;-" Output="SnapshotBranch"/>
+						<Property Name="JobURL" Value="" If="'$(UE_HORDE_JOBID)' == ''" />
+						<Property Name="JobURL" Value="https://horde.devtools.epicgames.com/job/$(UE_HORDE_JOBID)" If="'$(UE_HORDE_JOBID)' != ''" />
 						<!-- Publish to builds cloud storage endpoint -->
 						<Property Name="SnapshotBuildBaselineBranch" Value="$(Branch)"/>
 						<Property Name="SnapshotBuildBaselineBranch" Value="$(UE-CloudPublishBaselineBranch)" If="'$(UE-CloudPublishBaselineBranch)' != ''"/>
-						<Property Name="SnapshotBuildMetaData" Value="name=$(SnapshotIdentifier);depot=$(Depot);branch=$(Branch);baselineBranch=$(SnapshotBuildBaselineBranch);project=$(ProjectName);runtime=$(TargetType);platform=$(TargetPlatform);cookPlatform=$(CookPlatform);changelist=$(Change);codeChangelist=$(CodeChange);isPreflight=$(IsPreflight);isBuildMachine=$(IsBuildMachine);isHordeEnv=$(IsHordeEnv)"/>
+						<Property Name="SnapshotBuildMetaData" Value="name=$(SnapshotIdentifier);depot=$(Depot);branch=$(Branch);baselineBranch=$(SnapshotBuildBaselineBranch);project=$(ProjectName);runtime=$(TargetType);platform=$(TargetPlatform);cookPlatform=$(CookPlatform);changelist=$(Change);commit=$(Change);codeChangelist=$(CodeChange);codeCommit=$(CodeChange);job=$(UE_HORDE_JOBID);step=$(UE_HORDE_STEPID);buildurl=$(JobURL);type=oplog;isPreflight=$(IsPreflight);isBuildMachine=$(IsBuildMachine);isHordeEnv=$(IsHordeEnv)"/>
 						<Property Name="SnapshotBuildMetaData" Value="$(SnapshotBuildMetaData);cookFlavor=$(BCRCookFlavor)" If="'$(BCRCookFlavor)' != ''" />
 						<Property Name="SnapshotBuildMetaData" Value="$(SnapshotBuildMetaData);preflightChangelist=$(PreflightChange)" If="'$(PreflightChange)' != ''" />
 						<Property Name="SnapshotBuildMetaData" Value="$(SnapshotBuildMetaData);hordeTemplateId=$(UE_HORDE_TEMPLATEID)" If="'$(UE_HORDE_TEMPLATEID)' != ''" />
diff --git a/Engine/Build/Graph/Tests/IncrementalCompare.xml b/Engine/Build/Graph/Tests/IncrementalCompare.xml
--- a/Engine/Build/Graph/Tests/IncrementalCompare.xml
+++ b/Engine/Build/Graph/Tests/IncrementalCompare.xml
@@ -82,6 +82,8 @@
 				<Property Name="PreflightChange" Value=""/>
 				<Property Name="SnapshotBranch" Value="$(EscapedBranch)"/>
 				<StringOp Input="$(EscapedBranch)" Method="Replace" Arguments=".;-" Output="SnapshotBranch"/>
+				<Property Name="JobURL" Value="" If="'$(UE_HORDE_JOBID)' == ''" />
+				<Property Name="JobURL" Value="https://horde.devtools.epicgames.com/job/$(UE_HORDE_JOBID)" If="'$(UE_HORDE_JOBID)' != ''" />
 				<Property Name="StepId" Value="$(UE_HORDE_STEPID)"/>
 				<Property Name="BuildVersion" Value="$(EscapedBranch)-CL-$(Change)"/>
 				<Property Name="SnapshotIdentifier" Value="$(UE_HORDE_TEMPLATEID).$(BuildVersion).$(UE_HORDE_JOBID:-6).$(StepId)"/>
@@ -94,7 +96,7 @@
 				<StringOp Input="$(TargetPlatform)" Method="SplitFirst" Arguments="Client" Output="ClientPrefix"/>
 				<Property Name="TargetType" Value="Client" If="'$(ClientPrefix)' != ''"/>
 
-				<Property Name="SnapshotBuildMetaData" Value="name=$(SnapshotIdentifier).$(TargetPlatform);depot=$(Depot);branch=$(Branch);baselineBranch=$(SnapshotBuildBaselineBranch);project=$(ProjectName);runtime=$(TargetType);platform=Win64;cookPlatform=$(TargetPlatform);changelist=$(Change);codeChangelist=$(CodeChange);isPreflight=$(IsPreflight);isBuildMachine=$(IsBuildMachine);isHordeEnv=$(IsHordeEnv)"/>
+				<Property Name="SnapshotBuildMetaData" Value="name=$(SnapshotIdentifier).$(TargetPlatform);depot=$(Depot);branch=$(Branch);baselineBranch=$(SnapshotBuildBaselineBranch);project=$(ProjectName);runtime=$(TargetType);platform=Win64;cookPlatform=$(TargetPlatform);changelist=$(Change);commit=$(Change);codeChangelist=$(CodeChange);codeCommit=$(CodeChange);job=$(UE_HORDE_JOBID);step=$(UE_HORDE_STEPID);buildurl=$(JobURL);type=oplog;isPreflight=$(IsPreflight);isBuildMachine=$(IsBuildMachine);isHordeEnv=$(IsHordeEnv)"/>
 				<Property Name="SnapshotBuildMetaData" Value="$(SnapshotBuildMetaData);preflightChangelist=$(PreflightChange)" If="'$(PreflightChange)' != ''" />
 				<Property Name="SnapshotBuildMetaData" Value="$(SnapshotBuildMetaData);hordeTemplateId=$(UE_HORDE_TEMPLATEID)" If="'$(UE_HORDE_TEMPLATEID)' != ''" />
 				<Log Message="Exporting cook state with metadata '$(SnapshotBuildMetaData)'"/>
diff --git a/Engine/Source/Programs/AutomationTool/BuildGraph/Tasks/ZenExportSnapshotTask.cs b/Engine/Source/Programs/AutomationTool/BuildGraph/Tasks/ZenExportSnapshotTask.cs
--- a/Engine/Source/Programs/AutomationTool/BuildGraph/Tasks/ZenExportSnapshotTask.cs
+++ b/Engine/Source/Programs/AutomationTool/BuildGraph/Tasks/ZenExportSnapshotTask.cs
@@ -627,10 +627,13 @@ namespace AutomationTool.Tasks
 						string destinationId;
 						if (destinationStorageType == SnapshotStorageType.Builds)
 						{
+							StringBuilder metadata = new StringBuilder();
+							metadata.AppendFormat("createdAt={0}", DateTime.UtcNow.ToString("O"));
 							if (!String.IsNullOrEmpty(_parameters.Metadata))
 							{
-								exportSingleSourceCommandline.AppendFormat(" --builds-metadata \"{0}\"", _parameters.Metadata);
+								metadata.AppendFormat(";{0}", _parameters.Metadata);
 							}
+							exportSingleSourceCommandline.AppendFormat(" --builds-metadata \"{0}\"", metadata.ToString());
 							CbObjectId objectId = CbObjectId.NewObjectId();
 							destinationId = objectId.ToString().ToLowerInvariant();
 						}
