Detect editor requirements based on test name for target tests

#rnx

diff --git a/Engine/Build/Graph/Tasks/BuildAndTestProject.xml b/Engine/Build/Graph/Tasks/BuildAndTestProject.xml
--- a/Engine/Build/Graph/Tasks/BuildAndTestProject.xml
+++ b/Engine/Build/Graph/Tasks/BuildAndTestProject.xml
@@ -124,6 +124,7 @@
 	<Option Name="EditorPlatformsToSkipBootTests" DefaultValue="$(EditorPlatformsToSkipTests)" Description="Editor platforms that should skip boot tests."/>
 	<Option Name="TargetPlatformsToSkipBootTests" DefaultValue="$(TargetPlatformsToSkipTests)" Description="Target platforms that should skip boot tests."/>
 	<Option Name="CookedEditorPlatformsToSkipBootTests" DefaultValue="$(CookedEditorPlatformsToSkipTests)" Description="Cooked editor platforms that should skip boot tests."/>
+	<Option Name="RequiringEditorTargetTestList" DefaultValue="" Description="List of tests that require the editor to run. Test1+Test2 etc. Match is partial"/>
 
 	<!-- Similar to below but they are specifically added before the other tests -->
 	<!-- -->
@@ -184,6 +185,9 @@
 		<Expand Name="AddTestConfig" TestConfigName=""/>
 	</Do>
 
+	<!-- Add known target test types that require the editor -->
+	<Property Name="RequiringEditorTargetTestList" Value="$(RequiringEditorTargetTestList)+UE.TargetAutomation+Cook"/>
+
 	<!-- At this point we have all options so calculate a few things -->
 
 	<!-- Validate that the project exists-->
@@ -983,7 +987,6 @@
 						</Node>
 
 						<!-- Link the boot test in to the prerequisites and generated nodes if its enabled-->
-						<Property Name="TestPrerequisites" Value="$(TestPrerequisites);$(PlatformEditorNodeName)" />
 						<Do If="$(TargetBootTest) and !ContainsItem('$(TargetPlatformsToSkipBootTests)', '$(TargetPlatform)', '+')">
 							<Property Name="TestPrerequisites" Value="$(TestPrerequisites);$(BootTestName)" />
 							<Property Name="AllPlatformBootTestNodes" Value="$(AllPlatformBootTestNodes);$(BootTestName)"/>
@@ -1003,11 +1006,21 @@
 							<Property Name="TestNodeName" Value="$(PreNodeName)$(TargetPlatform) $(TestConfiguration)$(TestConfigTitle) $(TestNameTitle)" />
 							<Property Name="NodeAnnotations" Value=""/>
 							<Property Name="NodeAnnotations" Value="Workflow=$(AnnotationsTarget)" If="ContainsItem('$(AnnotationsPlatforms)', '$(TargetPlatform)', ';') and ContainsItem('$(AnnotationsTests)', '$(TestName)', ';')"/>
-							<!-- Run the test. No 'Requires' as we run on the same machine locally and use the publish folder otherwise -->
+
 							<Property Name="TestNameBeforeParentheses" Value="$(TestName)"/>
 							<Property Name="TestArgumentsWithinParentheses" Value=""/>
 							<Expand Name="SplitCombinedTestNameAndArgs"/>
-							<Node Name="$(TestNodeName)" Requires="$(TestPrerequisites)" Annotations="$(NodeAnnotations)">
+
+							<!-- Test 'Requires' should be the package game and the publish folder is used; however some tests may require the editor ie: UE.TargetAutomation -->
+							<Property Name="AdditionalTestPrerequisites" Value=""/>
+							<ForEach Name="ToMatchName" Values="$(RequiringEditorTargetTestList)" Separator="+">
+								<Do If="Contains('$(TestNameBeforeParentheses)', '$(ToMatchName)') And !ContainsItem('$(TestPrerequisites)', '$(PlatformEditorNodeName)', ';')">
+									<Property Name="AdditionalTestPrerequisites" Value="$(AdditionalTestPrerequisites);$(PlatformEditorNodeName)" />
+								</Do>
+							</ForEach>
+
+							<!-- 'Requires' should be the package game and the publish folder is used; however some tests may require the editor ie: UE.TargetAutomation -->
+							<Node Name="$(TestNodeName)" Requires="$(TestPrerequisites);$(AdditionalTestPrerequisites)" Annotations="$(NodeAnnotations)">
 								<Command Name="RunUnreal" Arguments="-project=$(TargetProject) -test=&quot;$(TestNameBeforeParentheses)&quot; -platform=$(TargetPlatform) -configuration=$(TestConfiguration) -target=$(TargetName) $(GauntletBuildArg) $(GauntletTestArgs) $(TestArgumentsWithinParentheses) $(GauntletReportArgs) $(TestConfigArgs)"/>
 							</Node>
 							<!-- Add it to our list -->
