Fixes for IncrementalCompare job:
- Use ProjectDir not ProjectName for bucket and project metadata
- Do oplog-export after oplog-mirror

diff --git a/Engine/Build/Graph/Tests/IncrementalCompare.xml b/Engine/Build/Graph/Tests/IncrementalCompare.xml
--- a/Engine/Build/Graph/Tests/IncrementalCompare.xml
+++ b/Engine/Build/Graph/Tests/IncrementalCompare.xml
@@ -73,6 +73,26 @@
 			<Commandlet Name="Cook" Project="$(ProjectFile)" Arguments="$(CommonCookArgs)"/>
 		</Node>
 		<Node Name="$(StoreIncrementalCookNodeName)" Requires="$(NaturalIncrementalCookNodeName)">
+			<!-- Move all cook artifacts to the compare_incremental dir. -->
+			<Log Message="$(NewLine)========== IncrementalCompare.xml: Move =========="/>
+			<Log Message="Delete Directories=$(CompareDirIncremental)"/>
+			<Delete Directories="$(CompareDirIncremental)"/>
+			<Do If="$(WithNullCook)">
+				<Log Message="Copy From=$(PlatformCookedDir)/*.* To=$(CompareDirIncremental)"/>
+				<Copy From="$(PlatformCookedDir)/*.*" To="$(CompareDirIncremental)" ErrorIfNotFound="true"/>
+			</Do>
+			<Do If="!$(WithNullCook)">
+				<Log Message="Move From=$(PlatformCookedDir)/*.* To=$(CompareDirIncremental)"/>
+				<Move From="$(PlatformCookedDir)/*.*" To="$(CompareDirIncremental)" ErrorIfNotFound="true"/>
+				<Log Message="Delete Directories=$(PlatformCookedDir)"/>
+				<Delete Directories="$(PlatformCookedDir)"/>
+			</Do>
+
+			<!-- Write the files from incremental cook's zenstore out into CompareDirIncremental. -->
+			<Log Message="$(NewLine)========== IncrementalCompare.xml: ZenMirror Incremental Cook =========="/>
+			<Log Message="ZenMirror Project=$(ProjectFile) Platform=$(TargetPlatform) DestinationFileDir=$(CompareDirIncremental)"/>
+			<ZenMirror Project="$(ProjectFile)" Platform="$(TargetPlatform)" DestinationFileDir="$(CompareDirIncremental)"/>
+			
 			<!-- Exporting snapshot for developer debugging needs -->
 			<Property Name="ShouldSnapshotPublishCloud" Value="false"/>
 			<Property Name="ShouldSnapshotPublishCloud" Value="true" If="$(PublishCloudSnapshot) and '$(UE-CloudPublishHost)' != ''"/>
@@ -94,31 +114,12 @@
 				<StringOp Input="$(TargetPlatform)" Method="SplitFirst" Arguments="Client" Output="ClientPrefix"/>
 				<Property Name="TargetType" Value="Client" If="'$(ClientPrefix)' != ''"/>
 
-				<Property Name="SnapshotBuildMetaData" Value="name=$(SnapshotIdentifier).$(TargetPlatform);depot=$(Depot);branch=$(Branch);baselineBranch=$(SnapshotBuildBaselineBranch);project=$(ProjectName);runtime=$(TargetType);platform=Win64;cookPlatform=$(TargetPlatform);changelist=$(Change);commit=$(Change);codeChangelist=$(CodeChange);codeCommit=$(CodeChange);isPreflight=$(IsPreflight);isBuildMachine=$(IsBuildMachine);isHordeEnv=$(IsHordeEnv)"/>
+				<Property Name="SnapshotBuildMetaData" Value="name=$(SnapshotIdentifier).$(TargetPlatform);depot=$(Depot);branch=$(Branch);baselineBranch=$(SnapshotBuildBaselineBranch);project=$(ProjectDir);runtime=$(TargetType);platform=Win64;cookPlatform=$(TargetPlatform);changelist=$(Change);commit=$(Change);codeChangelist=$(CodeChange);codeCommit=$(CodeChange);isPreflight=$(IsPreflight);isBuildMachine=$(IsBuildMachine);isHordeEnv=$(IsHordeEnv)"/>
 				<Property Name="SnapshotBuildMetaData" Value="$(SnapshotBuildMetaData);preflightChangelist=$(PreflightChange)" If="'$(PreflightChange)' != ''" />
 				<Property Name="SnapshotBuildMetaData" Value="$(SnapshotBuildMetaData);hordeTemplateId=$(UE_HORDE_TEMPLATEID)" If="'$(UE_HORDE_TEMPLATEID)' != ''" />
 				<Log Message="Exporting cook state with metadata '$(SnapshotBuildMetaData)'"/>
-				<ZenExportSnapshot Project="$(ProjectFile)" Platform="$(TargetPlatform)" DestinationStorageType="Builds" DestinationCloudHost="$(UE-CloudPublishHost)" DestinationCloudHttpVersion="$(UE-CloudPublishHttpVersion)" DestinationCloudNamespace="$(UE-CloudPublishNamespace)" DestinationCloudBucket="$(ProjectName).oplog-incrementalcompare.$(SnapshotBranch).$(TargetPlatform)" DestinationIdentifier="$(SnapshotIdentifier).$(TargetPlatform)" Metadata="$(SnapshotBuildMetaData)"/>
+				<ZenExportSnapshot Project="$(ProjectFile)" Platform="$(TargetPlatform)" DestinationStorageType="Builds" DestinationCloudHost="$(UE-CloudPublishHost)" DestinationCloudHttpVersion="$(UE-CloudPublishHttpVersion)" DestinationCloudNamespace="$(UE-CloudPublishNamespace)" DestinationCloudBucket="$(ProjectDir).oplog-incrementalcompare.$(SnapshotBranch).$(TargetPlatform)" DestinationIdentifier="$(SnapshotIdentifier).$(TargetPlatform)" Metadata="$(SnapshotBuildMetaData)"/>
 			</Do>
-			<!-- Move all cook artifacts to the compare_incremental dir. -->
-			<Log Message="$(NewLine)========== IncrementalCompare.xml: Move =========="/>
-			<Log Message="Delete Directories=$(CompareDirIncremental)"/>
-			<Delete Directories="$(CompareDirIncremental)"/>
-			<Do If="$(WithNullCook)">
-				<Log Message="Copy From=$(PlatformCookedDir)/*.* To=$(CompareDirIncremental)"/>
-				<Copy From="$(PlatformCookedDir)/*.*" To="$(CompareDirIncremental)" ErrorIfNotFound="true"/>
-			</Do>
-			<Do If="!$(WithNullCook)">
-				<Log Message="Move From=$(PlatformCookedDir)/*.* To=$(CompareDirIncremental)"/>
-				<Move From="$(PlatformCookedDir)/*.*" To="$(CompareDirIncremental)" ErrorIfNotFound="true"/>
-				<Log Message="Delete Directories=$(PlatformCookedDir)"/>
-				<Delete Directories="$(PlatformCookedDir)"/>
-			</Do>
-
-			<!-- Write the files from incremental cook's zenstore out into CompareDirIncremental. -->
-			<Log Message="$(NewLine)========== IncrementalCompare.xml: ZenMirror Incremental Cook =========="/>
-			<Log Message="ZenMirror Project=$(ProjectFile) Platform=$(TargetPlatform) DestinationFileDir=$(CompareDirIncremental)"/>
-			<ZenMirror Project="$(ProjectFile)" Platform="$(TargetPlatform)" DestinationFileDir="$(CompareDirIncremental)"/>
 		</Node>
 		<Property Name="FullCookPrequisiteNodeName" Value="$(StoreIncrementalCookNodeName)"/>
 		<Do If="$(WithNullCook)">
