Add "cookPlatform" metadata to created cloud artifacts of staged-build and packaged-build type.  The value for this field is consistently meant to match the target platform as known at cook time and stage time.

diff --git a/Engine/Build/Graph/Tasks/BuildAndTestProject.xml b/Engine/Build/Graph/Tasks/BuildAndTestProject.xml
--- a/Engine/Build/Graph/Tasks/BuildAndTestProject.xml
+++ b/Engine/Build/Graph/Tasks/BuildAndTestProject.xml
@@ -789,7 +789,7 @@
 					<Expand Name="WriteArtifactMetadataFile" Type="staged-build" Platform="$(StagedPlatformFolder)" Configuration="$(PackageConfigurations)" ArtifactDir="$(TargetDirectory)" />
 					<Tag Files="$(TargetDirectory)/..." With="#$(PlatformStageNodeName) Files"/>
 
-					<CreateCloudArtifact Name="$(BuildNamePath)-Staged" HordeArtifactName="Cloud-$(ProjectName)-Staged-$(StagedPlatformFolder)" Type="staged-build" BaseDir="$(TargetDirectory)" Metadata="dashboard-category=Staged Builds - Cloud" Project="$(TargetProject)" Platform="$(StagedPlatformFolder)" Files="#$(PlatformStageNodeName) Files"/>
+					<CreateCloudArtifact Name="$(BuildNamePath)-Staged" HordeArtifactName="Cloud-$(ProjectName)-Staged-$(StagedPlatformFolder)" Type="staged-build" BaseDir="$(TargetDirectory)" Metadata="cookPlatform=$(StagedPlatformFolder);dashboard-category=Staged Builds - Cloud" Project="$(TargetProject)" Platform="$(StagedPlatformFolder)" Files="#$(PlatformStageNodeName) Files"/>
 					<Delete Files="$(TargetDirectory)/artifactmetadata.json"/>
 				</Node>
 			</Do>
@@ -831,7 +831,7 @@
 						<Expand Name="WriteArtifactMetadataFile" Type="packaged-build" Platform="$(StagedPlatformFolder)" Configuration="$(PackageConfiguration)" ArtifactDir="$(TargetDirectory)" />
 						<Tag Files="$(TargetDirectory)/..." With="#$(PlatformPackageNodeName) $(PackageConfiguration) Files"/>
 
-						<CreateCloudArtifact Name="$(BuildNamePath)-Packaged-$(PackageConfiguration)" HordeArtifactName="Cloud-$(ProjectName)-Packaged-$(StagedPlatformFolder)-$(PackageConfiguration)" Type="packaged-build" BaseDir="$(TargetDirectory)" Metadata="Configuration=$(PackageConfiguration);dashboard-category=Packaged Builds - Cloud" Project="$(TargetProject)" Platform="$(StagedPlatformFolder)" Files="#$(PlatformPackageNodeName) $(PackageConfiguration) Files"/>
+						<CreateCloudArtifact Name="$(BuildNamePath)-Packaged-$(PackageConfiguration)" HordeArtifactName="Cloud-$(ProjectName)-Packaged-$(StagedPlatformFolder)-$(PackageConfiguration)" Type="packaged-build" BaseDir="$(TargetDirectory)" Metadata="cookPlatform=$(StagedPlatformFolder);Configuration=$(PackageConfiguration);dashboard-category=Packaged Builds - Cloud" Project="$(TargetProject)" Platform="$(StagedPlatformFolder)" Files="#$(PlatformPackageNodeName) $(PackageConfiguration) Files"/>
 						<Delete Files="$(TargetDirectory)/artifactmetadata.json"/>
 					</ForEach>
 				</Node>
