IncrementalCompare job : show instigators

#rnx

diff --git a/Engine/Build/Graph/Tests/IncrementalCompare.xml b/Engine/Build/Graph/Tests/IncrementalCompare.xml
--- a/Engine/Build/Graph/Tests/IncrementalCompare.xml
+++ b/Engine/Build/Graph/Tests/IncrementalCompare.xml
@@ -126,14 +126,14 @@
 			<Node Name="$(PrimingNullIncrementalCookNodeName)" Requires="$(StoreIncrementalCookNodeName)">
 				<!-- Execute the null incremental cook. -->
 				<Log Message="$(NewLine)========== IncrementalCompare.xml: Priming Null Incremental Cook =========="/>
-				<Property Name="NullCookArgs" Value="$(CommonCookArgs) -dpcvars=cook.displaymode=2"/>
+				<Property Name="NullCookArgs" Value="$(CommonCookArgs) -dpcvars=cook.displaymode=4"/>
 				<Property Name="NullCookArgs" Value="$(NullCookArgs) -CustomConfig=$(NullCookCustomConfig)" If="'$(NullCookCustomConfig)' != ''"/>
 				<Commandlet Name="Cook" Project="$(ProjectFile)" Arguments="$(NullCookArgs)"/>
 			</Node>
 			<Node Name="$(NullIncrementalCookNodeName)" Requires="$(PrimingNullIncrementalCookNodeName)">
 				<!-- Execute the null incremental cook. -->
 				<Log Message="$(NewLine)========== IncrementalCompare.xml: Null Incremental Cook =========="/>
-				<Property Name="NullCookArgs" Value="$(CommonCookArgs) -dpcvars=cook.displaymode=2"/>
+				<Property Name="NullCookArgs" Value="$(CommonCookArgs) -dpcvars=cook.displaymode=4"/>
 				<Property Name="NullCookArgs" Value="$(NullCookArgs) -CustomConfig=$(NullCookCustomConfig)" If="'$(NullCookCustomConfig)' != ''"/>
 				<Commandlet Name="Cook" Project="$(ProjectFile)" Arguments="$(NullCookArgs)"/>
 			</Node>
