Support Linux as an optional platform for BuildAndTestProject dedicated server builds
- Windows is still the default dedicated server platform  but will pass Linux to the server compile if the target platform is detected

#rnx

diff --git a/Engine/Build/Graph/Tasks/BuildAndTestProject.xml b/Engine/Build/Graph/Tasks/BuildAndTestProject.xml
--- a/Engine/Build/Graph/Tasks/BuildAndTestProject.xml
+++ b/Engine/Build/Graph/Tasks/BuildAndTestProject.xml
@@ -667,7 +667,9 @@
 						<ForEach Name="TargetConfiguration" Values="$(TargetConfigurations)" Separator="+">
 							<Compile Target="$(TargetName)" Platform="$(TargetPlatform)" Configuration="$(TargetConfiguration)" Arguments="-BuildVersion=&quot;$(BuildVersion)&quot; $(GenericCompileArguments) $(TargetCompileArguments)"/>
 							<Do If="'$(DedicatedServer)'" >
-								<Compile Target="$(TargetName)Server" Platform="Win64" Configuration="$(TargetConfiguration)" Arguments="-BuildVersion=&quot;$(BuildVersion)&quot; $(GenericCompileArguments) $(TargetCompileArguments)"/>
+								<Property Name="ServerPlatform" Value="Win64"/>
+								<Property Name="ServerPlatform" Value="Linux" If="$(IsTargetPlatformLinux)"/>
+								<Compile Target="$(TargetName)Server" Platform="$(ServerPlatform)" Configuration="$(TargetConfiguration)" Arguments="-BuildVersion=&quot;$(BuildVersion)&quot; $(GenericCompileArguments) $(TargetCompileArguments)"/>
 							</Do>
 						</ForEach>
 					</Default>
