Adding the ability to turn off network artifacts

#rnx

diff --git a/Engine/Build/Graph/Tasks/BuildAndTestProject.xml b/Engine/Build/Graph/Tasks/BuildAndTestProject.xml
--- a/Engine/Build/Graph/Tasks/BuildAndTestProject.xml
+++ b/Engine/Build/Graph/Tasks/BuildAndTestProject.xml
@@ -171,7 +171,7 @@
 	<Include Script="Inc/CommonProperties.xml" />
 
 	<!-- Common settings but can be overidden -->
-	<Option Name="UseNetworkShare" DefaultValue="$(IsBuildMachine)" Restrict="true|false" Description="Whether to use network share for build output"/>
+	<Option Name="UseNetworkShare" DefaultValue="$(IsBuildMachine)" Restrict="true|false" Description="This option is deprecated. Please use PublishToNetworkOutputDirectory  instead"/>
 	<Include Script="$(PathSettingsFile)" />
 	<Include Script="$(GauntletSettingsFile)" />
 	<Do If="'$(PreStepsExtendFile)' != ''">
@@ -230,6 +230,7 @@
 		<Property Name="ExtraCookArguments" Value="$(ExtraCookArguments) -cookincremental"/>
 	</Do>
 	<Option Name="RetrieveUsingCloudArtifacts" DefaultValue="false" Restrict="true|false" Description="Enable to retrieve cloud artifacts" />
+	<Option Name="PublishToNetworkOutputDirectory" DefaultValue="$(IsBuildMachine)" Restrict="true|false" Description="Enable to publish artifacts to network output directory" />
 	<Option Name="PublishCloudArtifacts" DefaultValue="$(RetrieveUsingCloudArtifacts)" Description="Enable to publish cloud artifacts" />
 	<Option Name="PublishHordeArtifacts" DefaultValue="$(IsEpicBuildMachine)" Description="Enable to publish horde artifacts"/>
 
@@ -449,7 +450,7 @@
 			<!-- Publish CookedEditor -->
 			<Node Name="$(PlatformPublishNodeName)" Requires="$(PreNodeName)Stage CookedEditor $(EditorPlatform)">
 				<Property Name="TargetDirectory" Value="$(ProjectOutputDirectory)/$(EditorPlatformPath)CookedEditor"/>
-				<Do If="$(UseNetworkShare)">
+				<Do If="$(PublishToNetworkOutputDirectory)">
 					<Log Message="Copying to $(NetworkOutputDirectory)/$(EditorPlatformPath)CookedEditor"/>
 					<Do If="!$(SkipPublish)">
 						<Copy Files="..." From="$(TargetDirectory)" To="$(NetworkOutputDirectory)/$(EditorPlatformPath)CookedEditor"/>
@@ -779,9 +780,12 @@
 			
 			<!-- Publish (Staged) -->
 			<Node Name="$(PlatformPublishStagedNodeName)" Requires="$(PlatformPublishStageDependencyName)">
-				<Error Message="NetworkOutputDirectory is not set" If="'$(NetworkOutputDirectory)' == ''"/>
-				<Copy Files="..." From="$(ProjectOutputDirectory)/$(StagedPlatformFolder)" To="$(NetworkOutputDirectory)/$(StagedPlatformFolder)/Staged"/>
-				<Copy Files="..." From="$(ShaderSymbolsDirectory)/*$(TargetPlatform)*" To="$(NetworkOutputDirectory)/$(StagedPlatformFolder)/ShaderSymbols"/>
+				<Log Message="Network upload is disabled, skipping upload." If="!$(PublishToNetworkOutputDirectory)"/>
+				<Do If="$(PublishToNetworkOutputDirectory)">				
+					<Error Message="NetworkOutputDirectory is not set" If="'$(NetworkOutputDirectory)' == ''"/>
+					<Copy Files="..." From="$(ProjectOutputDirectory)/$(StagedPlatformFolder)" To="$(NetworkOutputDirectory)/$(StagedPlatformFolder)/Staged"/>
+					<Copy Files="..." From="$(ShaderSymbolsDirectory)/*$(TargetPlatform)*" To="$(NetworkOutputDirectory)/$(StagedPlatformFolder)/ShaderSymbols"/>
+				</Do>
 			</Node>
 
 			<!-- Publish (Staged) (Cloud) -->
@@ -815,10 +819,13 @@
 			
 			<!-- Publish (Packages) -->
 			<Node Name="$(PlatformPublishPackagedNodeName)" Requires="$(PlatformPackageNodeName)">
-				<Error Message="NetworkOutputDirectory is not set" If="'$(NetworkOutputDirectory)' == ''"/>
-				<ForEach Name="PackageConfiguration" Values="$(PackageConfigurations)" Separator="+">
-					<Command Name="BuildCookRun" Arguments="$(BCRArgs) -configuration=$(PackageConfiguration) $(StageAndPackageArguments) -skipbuild -skipcook -skipstage -skippak -skippackage -archive -archivedirectory=&quot;$(NetworkOutputDirectory)/$(StagedPlatformFolder)/Packages&quot;" />
-				</ForEach>
+				<Log Message="Network upload is disabled, skipping upload." If="!$(PublishToNetworkOutputDirectory)"/>
+				<Do If="$(PublishToNetworkOutputDirectory)">
+					<Error Message="NetworkOutputDirectory is not set" If="'$(NetworkOutputDirectory)' == ''"/>
+					<ForEach Name="PackageConfiguration" Values="$(PackageConfigurations)" Separator="+">
+						<Command Name="BuildCookRun" Arguments="$(BCRArgs) -configuration=$(PackageConfiguration) $(StageAndPackageArguments) -skipbuild -skipcook -skipstage -skippak -skippackage -archive -archivedirectory=&quot;$(NetworkOutputDirectory)/$(StagedPlatformFolder)/Packages&quot;" />
+					</ForEach>
+				</Do>
 			</Node>
 
 			<!-- Publish (Packages) (Cloud) -->
diff --git a/Engine/Build/InstalledEngineBuild.xml b/Engine/Build/InstalledEngineBuild.xml
--- a/Engine/Build/InstalledEngineBuild.xml
+++ b/Engine/Build/InstalledEngineBuild.xml
@@ -36,7 +36,8 @@
 	<Option Name="ExtraCompileArgsMac" DefaultValue="" Description="Extra arguments to pass to compiles run on Mac platforms"/>
 	<Option Name="ExtraDDCArgs" DefaultValue="" Description="Extra arguments to pass to DDC Generation"/>
 
-	<Option Name="PublishCloudArtifacts" DefaultValue="false" Description="Enable to publish cloud artifacts" />
+	<Option Name="PublishCloudArtifacts" DefaultValue="false" Restrict="true|false" Description="Enable to publish cloud artifacts" />
+	<Option Name="PublishToNetworkOutputDirectory" DefaultValue="$(IsBuildMachine)" Restrict="true|false" Description="Enable to publish artifacts to network output directory" />
 
 	<Option Name="CloudPublishBaseline" DefaultValue="" Description="The baseline branch (your mainline) if incorrectly set your data will not de-dup across branches. Only applies if `PublishCloudArtifacts` is set." />
 	<Option Name="CloudPublishHost" DefaultValue="" Description="The host to publish cloud artifacts to if enabled using `PublishCloudArtifacts`" />
