IncrementalValidate: Copy ModifiedEditorPackages.txt and ModifiedCookedPackages.txt into the D:Engine\Programs\AutomationTool\Saved\Logs directory so they will be stored as artifacts in horde.
#rnx

diff --git a/Engine/Build/Graph/Tests/IncrementalValidate.xml b/Engine/Build/Graph/Tests/IncrementalValidate.xml
--- a/Engine/Build/Graph/Tests/IncrementalValidate.xml
+++ b/Engine/Build/Graph/Tests/IncrementalValidate.xml
@@ -30,6 +30,7 @@
 	<!-- Profiling Options -->
 	<Option Name="ProfilingFilesPublishRootPath" DefaultValue="" Description="Where to publish profiling data files."/>
 	<Option Name="WithCookUTRACECapture" DefaultValue="false" Restrict="true|false" Description="Whether to capture UTRACE data during cooks. Only true or false"/>
+	<Option Name="StoreHordeArtifacts" DefaultValue="true" Restrict="true|false" Description="Whether to copy diagnostic output files into Engine\Programs\AutomationTool\Saved\Logs for storage as Horde Artifacts"/>
 
 	<!-- Config Options -->
 	<Option Name="CustomConfig" DefaultValue="" Description="What custom configuration to apply (see -CustomConfig commandline argument)"/>
@@ -39,6 +40,8 @@
 	<Property Name="TraceArgs" Value=""/>
 	<Property Name="TraceArgs" Value="-trace=cpu,loadtime,savetime,cook,contextswitch" If="$(WithCookUTRACECapture) and '$(ProfilingFilesPublishRootPath)' != ''"/>
 	<Property Name="ProfilingFilesLocalRootPath" Value="$(RootDir)/Engine/Intermediate/Profiling"/>
+	<Property Name="CookedMetaDataDir" Value="$(ProjectDir)/Saved/Cooked/$(TargetPlatform)/$(ProjectName)/Metadata"/>
+	<Property Name="HordeArtifactDir" Value="$(RootDir)/Engine/Programs/AutomationTool/Saved/Logs"/>
 
 	<Property Name="CommonCookArgs" Value="-targetplatform=$(TargetPlatform) $(IncrementalArgs)"/>
 	<Property Name="CommonCookArgs" Value="$(CommonCookArgs) -CustomConfig=$(CustomConfig)" If="'$(CustomConfig)' != ''"/>
@@ -59,11 +62,15 @@
 			<Do If="'$(TraceArgsFinal)' != '' and Exists('$(ProfilingFilesLocalRootPath)')">
 				<Delete Files="$(ProfilingFilesLocalRootPath)/..."/>
 			</Do>
-			<Commandlet Name="Cook" Project="$(ProjectFile)" Arguments="$(CommonCookArgs) $(TraceArgsFinal) -IncrementalValidate -IncrementalValidateAllowWrite"/>
+			<Commandlet Name="Cook" Project="$(ProjectFile)" Arguments="$(CommonCookArgs) $(TraceArgsFinal) -IncrementalValidate -IncrementalValidateAllowWrite" ErrorLevel="999"/>
 			<Do If="'$(TraceArgsFinal)' != ''">
 				<Zip FromDir="$(ProfilingFilesLocalRootPath)" Files="*.utrace" ZipFile="$(ProfilingFilesLocalRootPath)/$(ProjectName)_IncrementalValidate.zip"/>
 				<Copy Files="$(ProjectName)_IncrementalValidate.zip" From="$(ProfilingFilesLocalRootPath)" To="$(ProfilingFilesPublishRootPath)\$(EscapedBranch)-CL-$(Change)"/>
 			</Do>
+			<Do If="'$(StoreHordeArtifacts)' == 'true'">
+				<Copy Files="ModifiedEditorPackages.txt" From="$(CookedMetaDataDir)" To="$(HordeArtifactDir)"/>
+				<Copy Files="ModifiedCookedPackages.txt" From="$(CookedMetaDataDir)" To="$(HordeArtifactDir)"/>
+			</Do>
 		</Node>
 		<Node Name="$(IncrementalValidatePhase1Name)" Requires="$(CompileNodeName)">
 			<Property Name="TraceArgsFinal" Value=""/>
@@ -71,16 +78,20 @@
 			<Do If="'$(TraceArgsFinal)' != '' and Exists('$(ProfilingFilesLocalRootPath)')">
 				<Delete Files="$(ProfilingFilesLocalRootPath)/..."/>
 			</Do>
-			<Commandlet Name="Cook" Project="$(ProjectFile)" Arguments="$(CommonCookArgs) $(TraceArgsFinal) -IncrementalValidatePhase1 -IncrementalValidateAllowWrite"/>
+			<Commandlet Name="Cook" Project="$(ProjectFile)" Arguments="$(CommonCookArgs) $(TraceArgsFinal) -IncrementalValidatePhase1 -IncrementalValidateAllowWrite" ErrorLevel="999"/>
 		</Node>
 		<Node Name="$(IncrementalValidatePhase2Name)" Requires="$(IncrementalValidatePhase1Name)">
 			<Property Name="TraceArgsFinal" Value=""/>
 			<Property Name="TraceArgsFinal" Value="$(TraceArgs) -tracefile=&quot;$(ProfilingFilesLocalRootPath)/$(ProjectName)_IncrementalValidatePhase2.utrace&quot;" If="$(WithCookUTRACECapture) and '$(ProfilingFilesPublishRootPath)' != ''"/>
-			<Commandlet Name="Cook" Project="$(ProjectFile)" Arguments="$(CommonCookArgs) $(TraceArgsFinal) -IncrementalValidatePhase2 -IncrementalValidateAllowWrite"/>
+			<Commandlet Name="Cook" Project="$(ProjectFile)" Arguments="$(CommonCookArgs) $(TraceArgsFinal) -IncrementalValidatePhase2 -IncrementalValidateAllowWrite" ErrorLevel="999"/>
 			<Do If="'$(TraceArgsFinal)' != ''">
 				<Zip FromDir="$(ProfilingFilesLocalRootPath)" Files="*.utrace" ZipFile="$(ProfilingFilesLocalRootPath)/$(ProjectName)_IncrementalValidate.zip"/>
 				<Copy Files="$(ProjectName)_IncrementalValidate.zip" From="$(ProfilingFilesLocalRootPath)" To="$(ProfilingFilesPublishRootPath)\$(EscapedBranch)-CL-$(Change)"/>
 			</Do>
+			<Do If="'$(StoreHordeArtifacts)' == 'true'">
+				<Copy Files="ModifiedEditorPackages.txt" From="$(CookedMetaDataDir)" To="$(HordeArtifactDir)"/>
+				<Copy Files="ModifiedCookedPackages.txt" From="$(CookedMetaDataDir)" To="$(HordeArtifactDir)"/>
+			</Do>
 		</Node>
 	</Agent>
 
