Allow incremental compare job to publish multiple snapshots at different steps to allow for easier debugging of differences in state between steps.

diff --git a/Engine/Build/Graph/Tests/IncrementalCompare.xml b/Engine/Build/Graph/Tests/IncrementalCompare.xml
--- a/Engine/Build/Graph/Tests/IncrementalCompare.xml
+++ b/Engine/Build/Graph/Tests/IncrementalCompare.xml
@@ -53,12 +53,44 @@
 	<Property Name="NaturalIncrementalCookNodeName" Value="Natural Incremental Cook $(ProjectName) $(TargetPlatform)"/>
 	<Property Name="StoreIncrementalCookNodeName" Value="Store Incremental Cook $(ProjectName) $(TargetPlatform)"/>
 	<Property Name="PrimingNullIncrementalCookNodeName" Value="Priming Null Incremental Cook $(ProjectName) $(TargetPlatform)"/>
+	<Property Name="PublishingPrimingNullIncrementalCookNodeName" Value="Publishing Priming Null Incremental Cook $(ProjectName) $(TargetPlatform)"/>
 	<Property Name="NullIncrementalCookNodeName" Value="Null Incremental Cook $(ProjectName) $(TargetPlatform)"/>
+	<Property Name="PublishingNullIncrementalCookNodeName" Value="Publishing Null Incremental Cook $(ProjectName) $(TargetPlatform)"/>
 	<Property Name="CleanNullIncrementalCookNodeName" Value="Clean Null Incremental Cook $(ProjectName) $(TargetPlatform)"/>
 	<Property Name="FullCookNodeName" Value="Full Cook $(ProjectName) $(TargetPlatform)"/>
 	<Property Name="StoreFullCookNodeName" Value="Store Full Cook $(ProjectName) $(TargetPlatform)"/>
 	<Property Name="DiffCooksNodeName" Value="Diff Cooks $(ProjectName) $(TargetPlatform)"/>
 
+	<Macro Name="Snapshot Publish" OptionalArguments="NameSuffix">
+		<Property Name="ShouldSnapshotPublishCloud" Value="false"/>
+		<Property Name="ShouldSnapshotPublishCloud" Value="true" If="$(PublishCloudSnapshot) and '$(UE-CloudPublishHost)' != ''"/>
+
+		<Do If="$(ShouldSnapshotPublishCloud)">
+			<Property Name="IsPreflight" Value="false"/>
+			<Property Name="PreflightChange" Value=""/>
+			<Property Name="SnapshotBranch" Value="$(EscapedBranch)"/>
+			<StringOp Input="$(EscapedBranch)" Method="Replace" Arguments=".;-" Output="SnapshotBranch"/>
+			<Property Name="StepId" Value="$(UE_HORDE_STEPID)"/>
+			<Property Name="BuildVersion" Value="$(EscapedBranch)-CL-$(Change)"/>
+			<Property Name="BuildVersion" Value="$(BuildVersion)-$(NameSuffix)" If="'$(NameSuffix)' !=''" />
+			<Property Name="SnapshotIdentifier" Value="$(UE_HORDE_TEMPLATEID).$(BuildVersion).$(UE_HORDE_JOBID:-6).$(StepId)"/>
+			<Property Name="SnapshotBuildBaselineBranch" Value="$(Branch)"/>
+			<Property Name="SnapshotBuildBaselineBranch" Value="$(UE-CloudPublishBaselineBranch)" If="'$(UE-CloudPublishBaselineBranch)' != ''"/>
+			<Property Name="IsHordeEnv" Value="false"/>
+			<Property Name="IsHordeEnv" Value="true" If="'$(UE_HORDE_JOBID)' != ''"/>
+			<Property Name="TargetType" Value="Server"/>
+			<Property Name="ClientPrefix" Value=""/>
+			<StringOp Input="$(TargetPlatform)" Method="SplitFirst" Arguments="Client" Output="ClientPrefix"/>
+			<Property Name="TargetType" Value="Client" If="'$(ClientPrefix)' != ''"/>
+
+			<Property Name="SnapshotBuildMetaData" Value="name=$(SnapshotIdentifier).$(TargetPlatform);depot=$(Depot);branch=$(Branch);baselineBranch=$(SnapshotBuildBaselineBranch);project=$(ProjectDir);runtime=$(TargetType);platform=Win64;cookPlatform=$(TargetPlatform);changelist=$(Change);commit=$(Change);codeChangelist=$(CodeChange);codeCommit=$(CodeChange);isPreflight=$(IsPreflight);isBuildMachine=$(IsBuildMachine);isHordeEnv=$(IsHordeEnv)"/>
+			<Property Name="SnapshotBuildMetaData" Value="$(SnapshotBuildMetaData);preflightChangelist=$(PreflightChange)" If="'$(PreflightChange)' != ''" />
+			<Property Name="SnapshotBuildMetaData" Value="$(SnapshotBuildMetaData);hordeTemplateId=$(UE_HORDE_TEMPLATEID)" If="'$(UE_HORDE_TEMPLATEID)' != ''" />
+			<Log Message="Exporting cook state with metadata '$(SnapshotBuildMetaData)'"/>
+			<ZenExportSnapshot Project="$(ProjectFile)" Platform="$(TargetPlatform)" DestinationStorageType="Builds" DestinationCloudHost="$(UE-CloudPublishHost)" DestinationCloudHttpVersion="$(UE-CloudPublishHttpVersion)" DestinationCloudNamespace="$(UE-CloudPublishNamespace)" DestinationCloudBucket="$(ProjectDir).oplog-incrementalcompare.$(SnapshotBranch).$(TargetPlatform)" DestinationIdentifier="$(SnapshotIdentifier).$(TargetPlatform)" Metadata="$(SnapshotBuildMetaData)"/>
+		</Do>
+	</Macro>
+
 	<Agent Name="IncrementalCompareAgent" Type="$(AgentType)">
 		<Node Name="$(CompileNodeName)">
 			<!-- Compile. -->
@@ -94,32 +126,7 @@
 			<ZenMirror Project="$(ProjectFile)" Platform="$(TargetPlatform)" DestinationFileDir="$(CompareDirIncremental)"/>
 			
 			<!-- Exporting snapshot for developer debugging needs -->
-			<Property Name="ShouldSnapshotPublishCloud" Value="false"/>
-			<Property Name="ShouldSnapshotPublishCloud" Value="true" If="$(PublishCloudSnapshot) and '$(UE-CloudPublishHost)' != ''"/>
-
-			<Do If="$(ShouldSnapshotPublishCloud)">
-				<Property Name="IsPreflight" Value="false"/>
-				<Property Name="PreflightChange" Value=""/>
-				<Property Name="SnapshotBranch" Value="$(EscapedBranch)"/>
-				<StringOp Input="$(EscapedBranch)" Method="Replace" Arguments=".;-" Output="SnapshotBranch"/>
-				<Property Name="StepId" Value="$(UE_HORDE_STEPID)"/>
-				<Property Name="BuildVersion" Value="$(EscapedBranch)-CL-$(Change)"/>
-				<Property Name="SnapshotIdentifier" Value="$(UE_HORDE_TEMPLATEID).$(BuildVersion).$(UE_HORDE_JOBID:-6).$(StepId)"/>
-				<Property Name="SnapshotBuildBaselineBranch" Value="$(Branch)"/>
-				<Property Name="SnapshotBuildBaselineBranch" Value="$(UE-CloudPublishBaselineBranch)" If="'$(UE-CloudPublishBaselineBranch)' != ''"/>
-				<Property Name="IsHordeEnv" Value="false"/>
-				<Property Name="IsHordeEnv" Value="true" If="'$(UE_HORDE_JOBID)' != ''"/>
-				<Property Name="TargetType" Value="Server"/>
-				<Property Name="ClientPrefix" Value=""/>
-				<StringOp Input="$(TargetPlatform)" Method="SplitFirst" Arguments="Client" Output="ClientPrefix"/>
-				<Property Name="TargetType" Value="Client" If="'$(ClientPrefix)' != ''"/>
-
-				<Property Name="SnapshotBuildMetaData" Value="name=$(SnapshotIdentifier).$(TargetPlatform);depot=$(Depot);branch=$(Branch);baselineBranch=$(SnapshotBuildBaselineBranch);project=$(ProjectDir);runtime=$(TargetType);platform=Win64;cookPlatform=$(TargetPlatform);changelist=$(Change);commit=$(Change);codeChangelist=$(CodeChange);codeCommit=$(CodeChange);isPreflight=$(IsPreflight);isBuildMachine=$(IsBuildMachine);isHordeEnv=$(IsHordeEnv)"/>
-				<Property Name="SnapshotBuildMetaData" Value="$(SnapshotBuildMetaData);preflightChangelist=$(PreflightChange)" If="'$(PreflightChange)' != ''" />
-				<Property Name="SnapshotBuildMetaData" Value="$(SnapshotBuildMetaData);hordeTemplateId=$(UE_HORDE_TEMPLATEID)" If="'$(UE_HORDE_TEMPLATEID)' != ''" />
-				<Log Message="Exporting cook state with metadata '$(SnapshotBuildMetaData)'"/>
-				<ZenExportSnapshot Project="$(ProjectFile)" Platform="$(TargetPlatform)" DestinationStorageType="Builds" DestinationCloudHost="$(UE-CloudPublishHost)" DestinationCloudHttpVersion="$(UE-CloudPublishHttpVersion)" DestinationCloudNamespace="$(UE-CloudPublishNamespace)" DestinationCloudBucket="$(ProjectDir).oplog-incrementalcompare.$(SnapshotBranch).$(TargetPlatform)" DestinationIdentifier="$(SnapshotIdentifier).$(TargetPlatform)" Metadata="$(SnapshotBuildMetaData)"/>
-			</Do>
+			<Expand Name="Snapshot Publish" NameSuffix="AfterNaturalCook"/>
 		</Node>
 		<Property Name="FullCookPrequisiteNodeName" Value="$(StoreIncrementalCookNodeName)"/>
 		<Do If="$(WithNullCook)">
@@ -130,14 +137,26 @@
 				<Property Name="NullCookArgs" Value="$(NullCookArgs) -CustomConfig=$(NullCookCustomConfig)" If="'$(NullCookCustomConfig)' != ''"/>
 				<Commandlet Name="Cook" Project="$(ProjectFile)" Arguments="$(NullCookArgs)"/>
 			</Node>
-			<Node Name="$(NullIncrementalCookNodeName)" Requires="$(PrimingNullIncrementalCookNodeName)">
+			<Node Name="$(PublishingPrimingNullIncrementalCookNodeName)" Requires="$(PrimingNullIncrementalCookNodeName)">
+				<!-- Publish the null incremental cook. -->
+				<Log Message="$(NewLine)========== IncrementalCompare.xml: Publishing primed Incremental Cook =========="/>
+				<!-- Exporting snapshot for developer debugging needs -->
+				<Expand Name="Snapshot Publish" NameSuffix="AfterPrimingNullCook"/>
+			</Node>
+			<Node Name="$(NullIncrementalCookNodeName)" Requires="$(PublishingPrimingNullIncrementalCookNodeName)">
 				<!-- Execute the null incremental cook. -->
 				<Log Message="$(NewLine)========== IncrementalCompare.xml: Null Incremental Cook =========="/>
 				<Property Name="NullCookArgs" Value="$(CommonCookArgs) -dpcvars=cook.displaymode=4"/>
 				<Property Name="NullCookArgs" Value="$(NullCookArgs) -CustomConfig=$(NullCookCustomConfig)" If="'$(NullCookCustomConfig)' != ''"/>
 				<Commandlet Name="Cook" Project="$(ProjectFile)" Arguments="$(NullCookArgs)"/>
 			</Node>
-			<Node Name="$(CleanNullIncrementalCookNodeName)" Requires="$(NullIncrementalCookNodeName)">
+			<Node Name="$(PublishingNullIncrementalCookNodeName)" Requires="$(NullIncrementalCookNodeName)">
+				<!-- Publish the null incremental cook. -->
+				<Log Message="$(NewLine)========== IncrementalCompare.xml: Publishing primed Incremental Cook =========="/>
+				<!-- Exporting snapshot for developer debugging needs -->
+				<Expand Name="Snapshot Publish" NameSuffix="AfterNullCook"/>
+			</Node>
+			<Node Name="$(CleanNullIncrementalCookNodeName)" Requires="$(PublishingNullIncrementalCookNodeName)">
 				<!-- Move all cook artifacts to the compare_incremental dir. -->
 				<Log Message="$(NewLine)========== IncrementalCompare.xml: Clean Null =========="/>
 				<Log Message="Delete Directories=$(PlatformCookedDir)"/>
