Changes to IncrementalCompare to allow further debugging of spurious package rebuilds in incremental cooks.

diff --git a/Engine/Build/Graph/Tests/IncrementalCompare.xml b/Engine/Build/Graph/Tests/IncrementalCompare.xml
--- a/Engine/Build/Graph/Tests/IncrementalCompare.xml
+++ b/Engine/Build/Graph/Tests/IncrementalCompare.xml
@@ -33,6 +33,7 @@
 	<Property Name="UBACompileArgs" Value="-UBA -ActionStallReportTime=300 -UBALinkRemote" If="$(WithUBA)"/>
 	
 	<!-- Agent=IncrementalCompareAgent variables -->
+	<Option Name="ExtraCookArgs" DefaultValue="" Description="Extra arguments to pass to all cook steps."/>
 	<Option Name="CookProcessCount" DefaultValue="" Description="Number of processes to use during cooking, default is ini-driven."/>
 	<Option Name="MPCookCoreSubscription" DefaultValue="" Description="Over or Under subscription factor for CPU cores to use during cooking, default is ini-driven."/>
 	<Option Name="WithNullCook" Restrict="true|false" DefaultValue="false" Description="Whether to do a null cook or not"/>
@@ -41,9 +42,11 @@
 
 	<Property Name="FullPathProjectDir" Value="$(RootDir)/$(ProjectDir)"/>
 	<Property Name="EditorPlatform" Value="Win64"/>
-	<Property Name="CommonCookArgs" Value="-targetplatform=$(TargetPlatform) -cookcultures=en -unversioned -zenstore -cookincremental -IgnoreBlueprintCookDependencies"/>
+	<Property Name="CommonCookArgs" Value="-targetplatform=$(TargetPlatform) -cookcultures=en -unversioned -zenstore -cookincremental"/>
 	<Property Name="CommonCookArgs" Value="$(CommonCookArgs) -CookProcessCount=$(CookProcessCount)" If="'$(CookProcessCount)' != ''"/>
 	<Property Name="CommonCookArgs" Value="$(CommonCookArgs) -MPCookCoreSubscription=$(MPCookCoreSubscription)" If="'$(MPCookCoreSubscription)' != ''"/>
+	<Property Name="CommonCookArgs" Value="$(CommonCookArgs) -CustomConfig=$(NullCookCustomConfig)" If="'$(WithNullCook) and $(NullCookCustomConfig)' != ''"/>
+	<Property Name="CommonCookArgs" Value="$(CommonCookArgs) $(ExtraCookArgs)" If="'$(ExtraCookArgs)' != ''"/>
 	<Property Name="PlatformCookedDir" Value="$(FullPathProjectDir)/Saved/Cooked/$(TargetPlatform)"/>
 	<Property Name="CompareDirIncremental" Value="$(FullPathProjectDir)/Intermediate/IncrementalCompare/Iter/$(TargetPlatform)"/>
 	<Property Name="CompareDirRecook" Value="$(FullPathProjectDir)/Intermediate/IncrementalCompare/Recook/$(TargetPlatform)"/>
@@ -134,7 +137,6 @@
 				<!-- Execute the null incremental cook. -->
 				<Log Message="$(NewLine)========== IncrementalCompare.xml: Priming Null Incremental Cook =========="/>
 				<Property Name="NullCookArgs" Value="$(CommonCookArgs) -dpcvars=cook.displaymode=4"/>
-				<Property Name="NullCookArgs" Value="$(NullCookArgs) -CustomConfig=$(NullCookCustomConfig)" If="'$(NullCookCustomConfig)' != ''"/>
 				<Commandlet Name="Cook" Project="$(ProjectFile)" Arguments="$(NullCookArgs)"/>
 			</Node>
 			<Node Name="$(PublishingPrimingNullIncrementalCookNodeName)" Requires="$(PrimingNullIncrementalCookNodeName)">
@@ -147,7 +149,6 @@
 				<!-- Execute the null incremental cook. -->
 				<Log Message="$(NewLine)========== IncrementalCompare.xml: Null Incremental Cook =========="/>
 				<Property Name="NullCookArgs" Value="$(CommonCookArgs) -dpcvars=cook.displaymode=4"/>
-				<Property Name="NullCookArgs" Value="$(NullCookArgs) -CustomConfig=$(NullCookCustomConfig)" If="'$(NullCookCustomConfig)' != ''"/>
 				<Commandlet Name="Cook" Project="$(ProjectFile)" Arguments="$(NullCookArgs)"/>
 			</Node>
 			<Node Name="$(PublishingNullIncrementalCookNodeName)" Requires="$(NullIncrementalCookNodeName)">
