Flexibility to support skipping editor compiles via SkipEditorBuild argument

#rnx

diff --git a/Engine/Build/Graph/Tasks/BuildAndTestProject.xml b/Engine/Build/Graph/Tasks/BuildAndTestProject.xml
--- a/Engine/Build/Graph/Tasks/BuildAndTestProject.xml
+++ b/Engine/Build/Graph/Tasks/BuildAndTestProject.xml
@@ -106,6 +106,7 @@
 	<!-- Customizations -->
 	<Option Name="Versioned" Restrict="true|false" DefaultValue="false" Description="Whether to embed changelist number into binaries"/>
 	<Option Name="SkipBuild" DefaultValue="false" Description="Whether to skip building" />
+	<Option Name="SkipEditorBuild" DefaultValue="false" Description="Whether to skip building the editor" />
 	<Option Name="SkipCook" DefaultValue="false" Description="Whether to skip cooking and packaging" />
 	<Option Name="SkipPackage" DefaultValue="false" Description="Whether to skip packaging" />
 	<Option Name="SkipPublish" DefaultValue="false" Description="Whether to skip cooking, packaging, publishing" />
@@ -346,7 +347,8 @@
 				</ForEach>
 			</Node>
 			<!--Add these nodes to our dependency list -->
-			<Property Name="GeneratedToolNodes" Value="$(CompileEditor);$(PreNodeName)Compile Tools $(EditorPlatform)" If="!$(SkipBuild)" />
+			<Property Name="GeneratedToolNodes" Value="$(CompileEditor)" If="!$(SkipEditorBuild)" />
+			<Property Name="GeneratedToolNodes" Value="$(GeneratedToolNodes);$(PreNodeName)Compile Tools $(EditorPlatform)" If="!$(SkipBuild)" />
 			<!-- Editor Labels -->
 			<Label Category="Editors" Name="$(PreLabelName)Editor $(EditorPlatform)" Requires="$(GeneratedToolNodes)" />
 		</Agent>
@@ -687,6 +689,7 @@
 
 		</Agent>
 
+		<Property Name="PlatformEditorNodeName" Value="" If="$(SkipEditorBuild)" />
 		<!-- Separate agent for cooking / staging, so the cook can happen in parallel with the client compile -->
 		<Agent Name="Package $(PlatformAgentName)" Type="$(PackageAgentType)">
 
