Cleaning up cloud upload step to remove buildcookrun call

#rnx

diff --git a/Engine/Build/Graph/Tasks/BuildAndTestProject.xml b/Engine/Build/Graph/Tasks/BuildAndTestProject.xml
--- a/Engine/Build/Graph/Tasks/BuildAndTestProject.xml
+++ b/Engine/Build/Graph/Tasks/BuildAndTestProject.xml
@@ -833,13 +833,9 @@
 
 			<!-- Publish (Packages) (Cloud) -->
 			<Do  If="$(PublishCloudArtifacts)">
-				<Property Name="CloudPackageSkipBCRArguments" Value="-stage -package"/>
-				<Property Name="CloudPackageSkipBCRArguments" Value="-skipstage -skippackage" If="'$(TargetPlatform)' == 'Mac'"/>
-
 				<Node Name="$(CloudArtifactPlatformPackagedNodeName)" Requires="$(PlatformPackageNodeName)">
 					<ForEach Name="PackageConfiguration" Values="$(PackageConfigurations)" Separator="+">
-						<Property Name="TargetDirectory" Value="$(ProjectOutputDirectory)/Packages/$(StagedPlatformFolder)/$(PackageConfiguration)"/>
-						<Command Name="BuildCookRun" Arguments="$(BCRArgs) -configuration=$(PackageConfiguration) $(StageAndPackageArguments) -skipbuild -skipcook -skippak $(CloudPackageSkipBCRArguments) -archive -archivedirectory=&quot;$(TargetDirectory)&quot;" />
+						<Property Name="TargetDirectory" Value="$(ProjectOutputDirectory)/Packages/$(StagedPlatformFolder)/$(PackageConfiguration)"/>						
 						<Expand Name="WriteArtifactMetadataFile" Type="packaged-build" Platform="$(StagedPlatformFolder)" Configuration="$(PackageConfiguration)" ArtifactDir="$(TargetDirectory)" />
 						<Tag Files="$(TargetDirectory)/..." With="#$(PlatformPackageNodeName) $(PackageConfiguration) Files"/>
 
