IncrementalValidate: Fix the copying of files from the metadata directory to use the ProjectDirLeafName instead of the ProjectName for the name of the parent directory of Metadata.
#rnx

diff --git a/Engine/Build/Graph/Tests/IncrementalValidate.xml b/Engine/Build/Graph/Tests/IncrementalValidate.xml
--- a/Engine/Build/Graph/Tests/IncrementalValidate.xml
+++ b/Engine/Build/Graph/Tests/IncrementalValidate.xml
@@ -36,11 +36,26 @@
 	<Option Name="CustomConfig" DefaultValue="" Description="What custom configuration to apply (see -CustomConfig commandline argument)"/>
 
 	<!-- Agent=IncrementalValidate variables -->
+	<!--
+	     Find the Directory name of the .uproject file, since this is the name used in the Cooked output as the parent of the MetaData directory
+	     TODO: Parse this from the leafname of $(ProjectDir) instead of guessing that it is $(ProjectName) or $(ProjectName)Game.
+	     We need a path function API to be able to do that parsing, and BuildGraph doesn't have an API for that yet.
+	 -->
+	<Property Name="ProjectDirLeafName"/>
+	<Property Name="ProjectDirLeafName" Value="$(ProjectName)" If="Exists('$(ProjectDir)\$(ProjectName).uproject')"/>
+	<Property Name="ProjectDirLeafName" Value="$(ProjectName)Game" If="Exists('$(ProjectDir)\$(ProjectName)Game.uproject')"/>
+	<Do If="'$(StoreHordeArtifacts)' == 'true'">
+		<Warning Message="Could not find ProjectDirLeafName for project $(ProjectName). ModifiedCookedPackages.txt file will not be copied to HordeArtifactDir."
+		 If="'$(ProjectDirLeafName)'==''"/>
+		<Warning Message="ProjectDir not specified. ModifiedCookedPackages.txt file will not be copied to HordeArtifactDir."
+		 If="'$(ProjectDir)'==''"/>
+	</Do>
 	<Property Name="EditorPlatform" Value="Win64"/>
 	<Property Name="TraceArgs" Value=""/>
 	<Property Name="TraceArgs" Value="-trace=cpu,loadtime,savetime,cook,contextswitch" If="$(WithCookUTRACECapture) and '$(ProfilingFilesPublishRootPath)' != ''"/>
 	<Property Name="ProfilingFilesLocalRootPath" Value="$(RootDir)/Engine/Intermediate/Profiling"/>
-	<Property Name="CookedMetaDataDir" Value="$(ProjectDir)/Saved/Cooked/$(TargetPlatform)/$(ProjectName)/Metadata"/>
+	<Property Name="CookedMetaDataDir"/>
+	<Property Name="CookedMetaDataDir" Value="$(ProjectDir)/Saved/Cooked/$(TargetPlatform)/$(ProjectDirLeafName)/Metadata" If="'$(ProjectDir)' != '' and '$(ProjectDirLeafName)' != ''"/>
 	<Property Name="HordeArtifactDir" Value="$(RootDir)/Engine/Programs/AutomationTool/Saved/Logs"/>
 
 	<Property Name="CommonCookArgs" Value="-targetplatform=$(TargetPlatform) $(IncrementalArgs)"/>
@@ -67,7 +82,7 @@
 				<Zip FromDir="$(ProfilingFilesLocalRootPath)" Files="*.utrace" ZipFile="$(ProfilingFilesLocalRootPath)/$(ProjectName)_IncrementalValidate.zip"/>
 				<Copy Files="$(ProjectName)_IncrementalValidate.zip" From="$(ProfilingFilesLocalRootPath)" To="$(ProfilingFilesPublishRootPath)\$(EscapedBranch)-CL-$(Change)"/>
 			</Do>
-			<Do If="'$(StoreHordeArtifacts)' == 'true'">
+			<Do If="'$(StoreHordeArtifacts)' == 'true' and '$(CookedMetaDataDir)' != ''">
 				<Copy Files="ModifiedEditorPackages.txt" From="$(CookedMetaDataDir)" To="$(HordeArtifactDir)"/>
 				<Copy Files="ModifiedCookedPackages.txt" From="$(CookedMetaDataDir)" To="$(HordeArtifactDir)"/>
 			</Do>
@@ -88,7 +103,7 @@
 				<Zip FromDir="$(ProfilingFilesLocalRootPath)" Files="*.utrace" ZipFile="$(ProfilingFilesLocalRootPath)/$(ProjectName)_IncrementalValidate.zip"/>
 				<Copy Files="$(ProjectName)_IncrementalValidate.zip" From="$(ProfilingFilesLocalRootPath)" To="$(ProfilingFilesPublishRootPath)\$(EscapedBranch)-CL-$(Change)"/>
 			</Do>
-			<Do If="'$(StoreHordeArtifacts)' == 'true'">
+			<Do If="'$(StoreHordeArtifacts)' == 'true' and '$(CookedMetaDataDir)' != ''">
 				<Copy Files="ModifiedEditorPackages.txt" From="$(CookedMetaDataDir)" To="$(HordeArtifactDir)"/>
 				<Copy Files="ModifiedCookedPackages.txt" From="$(CookedMetaDataDir)" To="$(HordeArtifactDir)"/>
 			</Do>
