Fix for breakage to Incremental Compare job following CL44345160 due to overriding the ProjectDir option.

diff --git a/Engine/Build/Graph/Tests/IncrementalCompare.xml b/Engine/Build/Graph/Tests/IncrementalCompare.xml
--- a/Engine/Build/Graph/Tests/IncrementalCompare.xml
+++ b/Engine/Build/Graph/Tests/IncrementalCompare.xml
@@ -11,15 +11,15 @@
 	<EnvVar Name="UE-CloudPublishBaselineBranch"/>
 
 	<!-- Project Options -->
-	<Option Name="ProjectFile" DefaultValue="_Unknown_" Restrict=".+" Description="The projects to cook"/>
 	<Option Name="ProjectName" DefaultValue="_Unknown_" Restrict=".+" Description="The projects to cook"/>
-	<Option Name="ProjectDir" DefaultValue="_Unknown_" Restrict=".+" Description="The projects to cook"/>
+	<Option Name="RelativeProjectDir" DefaultValue="_Unknown_" Restrict=".+" Description="The relative path for the project directory"/>
+	<Option Name="RelativeProjectFile" DefaultValue="_Unknown_" Restrict=".+" Description="The relative path for the project file"/>
 	<Option Name="EditorTarget" DefaultValue="$(ProjectName)Editor" Description="Compile name of the editor target"/>
 
 	<!-- Note we can not make the default='', because the default is used for some validation. -->
-	<Error Message="No project file specified" If="'$(ProjectFile)' == '_Unknown_'"/>
 	<Error Message="No project name specified" If="'$(ProjectName)' == '_Unknown_'"/>
-	<Error Message="No project dir specified" If="'$(ProjectDir)' == '_Unknown_'"/>
+	<Error Message="No relative project dir specified" If="'$(RelativeProjectDir)' == '_Unknown_'"/>
+	<Error Message="No relative project file specified" If="'$(RelativeProjectFile)' == '_Unknown_'"/>
 
 	<!-- Platform Options -->
 	<Option Name="TargetPlatform" DefaultValue="Windows" Restrict=".+" Description="The platform to cook"/>
@@ -40,7 +40,7 @@
 	<Option Name="NullCookCustomConfig" DefaultValue="" Description="What custom configuration to apply if doing a null cook (see -CustomConfig commandline argument)"/>
 	<Option Name="PublishCloudSnapshot" DefaultValue="false" Restrict="true|false" Description="Publish build to cloud host" />
 
-	<Property Name="FullPathProjectDir" Value="$(RootDir)/$(ProjectDir)"/>
+	<Property Name="FullPathProjectDir" Value="$(RootDir)/$(RelativeProjectDir)"/>
 	<Property Name="EditorPlatform" Value="Win64"/>
 	<Property Name="CommonCookArgs" Value="-targetplatform=$(TargetPlatform) -cookcultures=en -unversioned -zenstore -cookincremental"/>
 	<Property Name="CommonCookArgs" Value="$(CommonCookArgs) -CookProcessCount=$(CookProcessCount)" If="'$(CookProcessCount)' != ''"/>
@@ -86,11 +86,11 @@
 			<StringOp Input="$(TargetPlatform)" Method="SplitFirst" Arguments="Client" Output="ClientPrefix"/>
 			<Property Name="TargetType" Value="Client" If="'$(ClientPrefix)' != ''"/>
 
-			<Property Name="SnapshotBuildMetaData" Value="name=$(SnapshotIdentifier).$(TargetPlatform);depot=$(Depot);branch=$(Branch);baselineBranch=$(SnapshotBuildBaselineBranch);project=$(ProjectDir);runtime=$(TargetType);platform=Win64;cookPlatform=$(TargetPlatform);changelist=$(Change);commit=$(Change);codeChangelist=$(CodeChange);codeCommit=$(CodeChange);isPreflight=$(IsPreflight);isBuildMachine=$(IsBuildMachine);isHordeEnv=$(IsHordeEnv)"/>
+			<Property Name="SnapshotBuildMetaData" Value="name=$(SnapshotIdentifier).$(TargetPlatform);depot=$(Depot);branch=$(Branch);baselineBranch=$(SnapshotBuildBaselineBranch);project=$(RelativeProjectDir);runtime=$(TargetType);platform=Win64;cookPlatform=$(TargetPlatform);changelist=$(Change);commit=$(Change);codeChangelist=$(CodeChange);codeCommit=$(CodeChange);isPreflight=$(IsPreflight);isBuildMachine=$(IsBuildMachine);isHordeEnv=$(IsHordeEnv)"/>
 			<Property Name="SnapshotBuildMetaData" Value="$(SnapshotBuildMetaData);preflightChangelist=$(PreflightChange)" If="'$(PreflightChange)' != ''" />
 			<Property Name="SnapshotBuildMetaData" Value="$(SnapshotBuildMetaData);hordeTemplateId=$(UE_HORDE_TEMPLATEID)" If="'$(UE_HORDE_TEMPLATEID)' != ''" />
 			<Log Message="Exporting cook state with metadata '$(SnapshotBuildMetaData)'"/>
-			<ZenExportSnapshot Project="$(ProjectFile)" Platform="$(TargetPlatform)" DestinationStorageType="Builds" DestinationCloudHost="$(UE-CloudPublishHost)" DestinationCloudHttpVersion="$(UE-CloudPublishHttpVersion)" DestinationCloudNamespace="$(UE-CloudPublishNamespace)" DestinationCloudBucket="$(ProjectDir).oplog-incrementalcompare.$(SnapshotBranch).$(TargetPlatform)" DestinationIdentifier="$(SnapshotIdentifier).$(TargetPlatform)" Metadata="$(SnapshotBuildMetaData)"/>
+			<ZenExportSnapshot Project="$(RelativeProjectFile)" Platform="$(TargetPlatform)" DestinationStorageType="Builds" DestinationCloudHost="$(UE-CloudPublishHost)" DestinationCloudHttpVersion="$(UE-CloudPublishHttpVersion)" DestinationCloudNamespace="$(UE-CloudPublishNamespace)" DestinationCloudBucket="$(RelativeProjectDir).oplog-incrementalcompare.$(SnapshotBranch).$(TargetPlatform)" DestinationIdentifier="$(SnapshotIdentifier).$(TargetPlatform)" Metadata="$(SnapshotBuildMetaData)"/>
 		</Do>
 	</Macro>
 
@@ -98,14 +98,14 @@
 		<Node Name="$(CompileNodeName)">
 			<!-- Compile. -->
 			<Log Message="$(NewLine)========== IncrementalCompare.xml: Compile =========="/>
-			<Compile Target="ShaderCompileWorker" Platform="$(EditorPlatform)" Project="$(ProjectFile)" Configuration="Development" Arguments="$(UBACompileArgs)"/>
+			<Compile Target="ShaderCompileWorker" Platform="$(EditorPlatform)" Project="$(RelativeProjectFile)" Configuration="Development" Arguments="$(UBACompileArgs)"/>
 			<Compile Target="ZenLaunch" Platform="$(EditorPlatform)" Configuration="Development" Arguments="$(UBACompileArgs)"/>
-			<Compile Target="$(EditorTarget)" Platform="$(EditorPlatform)" Project="$(ProjectFile)" Configuration="Development" Arguments="$(UBACompileArgs)"/>
+			<Compile Target="$(EditorTarget)" Platform="$(EditorPlatform)" Project="$(RelativeProjectFile)" Configuration="Development" Arguments="$(UBACompileArgs)"/>
 		</Node>
 		<Node Name="$(NaturalIncrementalCookNodeName)" Requires="$(CompileNodeName)">
 			<!-- Execute the incremental cook (we don't pass -forcerecook). -->
 			<Log Message="$(NewLine)========== IncrementalCompare.xml: Incremental Cook =========="/>
-			<Commandlet Name="Cook" Project="$(ProjectFile)" Arguments="$(CommonCookArgs)"/>
+			<Commandlet Name="Cook" Project="$(RelativeProjectFile)" Arguments="$(CommonCookArgs)"/>
 		</Node>
 		<Node Name="$(StoreIncrementalCookNodeName)" Requires="$(NaturalIncrementalCookNodeName)">
 			<!-- Move all cook artifacts to the compare_incremental dir. -->
@@ -125,8 +125,8 @@
 
 			<!-- Write the files from incremental cook's zenstore out into CompareDirIncremental. -->
 			<Log Message="$(NewLine)========== IncrementalCompare.xml: ZenMirror Incremental Cook =========="/>
-			<Log Message="ZenMirror Project=$(ProjectFile) Platform=$(TargetPlatform) DestinationFileDir=$(CompareDirIncremental)"/>
-			<ZenMirror Project="$(ProjectFile)" Platform="$(TargetPlatform)" DestinationFileDir="$(CompareDirIncremental)"/>
+			<Log Message="ZenMirror Project=$(RelativeProjectFile) Platform=$(TargetPlatform) DestinationFileDir=$(CompareDirIncremental)"/>
+			<ZenMirror Project="$(RelativeProjectFile)" Platform="$(TargetPlatform)" DestinationFileDir="$(CompareDirIncremental)"/>
 			
 			<!-- Exporting snapshot for developer debugging needs -->
 			<Expand Name="Snapshot Publish" NameSuffix="AfterNaturalCook"/>
@@ -137,7 +137,7 @@
 				<!-- Execute the null incremental cook. -->
 				<Log Message="$(NewLine)========== IncrementalCompare.xml: Priming Null Incremental Cook =========="/>
 				<Property Name="NullCookArgs" Value="$(CommonCookArgs) -dpcvars=cook.displaymode=4"/>
-				<Commandlet Name="Cook" Project="$(ProjectFile)" Arguments="$(NullCookArgs)"/>
+				<Commandlet Name="Cook" Project="$(RelativeProjectFile)" Arguments="$(NullCookArgs)"/>
 			</Node>
 			<Node Name="$(PublishingPrimingNullIncrementalCookNodeName)" Requires="$(PrimingNullIncrementalCookNodeName)">
 				<!-- Publish the null incremental cook. -->
@@ -149,7 +149,7 @@
 				<!-- Execute the null incremental cook. -->
 				<Log Message="$(NewLine)========== IncrementalCompare.xml: Null Incremental Cook =========="/>
 				<Property Name="NullCookArgs" Value="$(CommonCookArgs) -dpcvars=cook.displaymode=4"/>
-				<Commandlet Name="Cook" Project="$(ProjectFile)" Arguments="$(NullCookArgs)"/>
+				<Commandlet Name="Cook" Project="$(RelativeProjectFile)" Arguments="$(NullCookArgs)"/>
 			</Node>
 			<Node Name="$(PublishingNullIncrementalCookNodeName)" Requires="$(NullIncrementalCookNodeName)">
 				<!-- Publish the null incremental cook. -->
@@ -168,7 +168,7 @@
 		<Node Name="$(FullCookNodeName)" Requires="$(FullCookPrequisiteNodeName)">
 			<!-- Execute the forcerecook cook. -->
 			<Log Message="$(NewLine)========== IncrementalCompare.xml: Non-Incremental Cook =========="/>
-			<Commandlet Name="Cook" Project="$(ProjectFile)" Arguments="$(CommonCookArgs) -forcerecook"/>
+			<Commandlet Name="Cook" Project="$(RelativeProjectFile)" Arguments="$(CommonCookArgs) -forcerecook"/>
 		</Node>
 		<Node Name="$(StoreFullCookNodeName)" Requires="$(FullCookNodeName)">
 			<!-- Copy all cook artifacts to the compare_incremental dir. Copy rather than move so they remain
@@ -181,13 +181,13 @@
 
 			<!-- Write the files from forcerecook's zenstore out into CompareDirRecook. -->
 			<Log Message="$(NewLine)========== IncrementalCompare.xml: ZenMirror ForceRecook Cook =========="/>
-			<Log Message="ZenMirror Project=$(ProjectFile) Platform=$(TargetPlatform) DestinationFileDir=$(CompareDirRecook)"/>
-			<ZenMirror Project="$(ProjectFile)" Platform="$(TargetPlatform)" DestinationFileDir="$(CompareDirRecook)"/>
+			<Log Message="ZenMirror Project=$(RelativeProjectFile) Platform=$(TargetPlatform) DestinationFileDir=$(CompareDirRecook)"/>
+			<ZenMirror Project="$(RelativeProjectFile)" Platform="$(TargetPlatform)" DestinationFileDir="$(CompareDirRecook)"/>
 		</Node>
 		<Node Name="$(DiffCooksNodeName)" Requires="$(StoreFullCookNodeName)">
 			<!-- Compare the incremental cook with the non-incremental cook. -->
 			<Log Message="$(NewLine)========== IncrementalCompare.xml: DiffCook =========="/>
-			<Commandlet Name="DiffCook" Project="$(ProjectFile)" Arguments="-base=$(CompareDirRecook) -comp=$(CompareDirIncremental) -targetplatform=$(TargetPlatform) -addedverbosity=error -removedverbosity=error -modifiedverbosity=warning -hide=modifiedpackages"/>
+			<Commandlet Name="DiffCook" Project="$(RelativeProjectFile)" Arguments="-base=$(CompareDirRecook) -comp=$(CompareDirIncremental) -targetplatform=$(TargetPlatform) -addedverbosity=error -removedverbosity=error -modifiedverbosity=warning -hide=modifiedpackages"/>
 		</Node>
 	</Agent>
 
