BuildAndTestProject.xml fixes for DedicatedServer=True
- new separate steps for server (compile, cook, stage, package, publish)
- option to package and stage server separately from client
- update DstmDemo scheduled build to match dev team expectations
#rnx

diff --git a/Engine/Build/Graph/Tasks/BuildAndTestProject.xml b/Engine/Build/Graph/Tasks/BuildAndTestProject.xml
--- a/Engine/Build/Graph/Tasks/BuildAndTestProject.xml
+++ b/Engine/Build/Graph/Tasks/BuildAndTestProject.xml
@@ -46,6 +46,7 @@
 	<Option Name="MapIniSection" DefaultValue="" Restrict="" Description="MapIniSection argument to pass when cooking" />
 	<Option Name="ExtraCookArguments" DefaultValue="" Restrict="" Description="Extra arguments to use when cooking" />
 	<Option Name="ExtraStageAndPackageArguments" DefaultValue="" Restrict="" Description="Extra arguments to use when staging and packaging" />
+	<Option Name="ExtraStageAndPackageServerArguments" DefaultValue="" Restrict="" Description="Extra arguments to use when staging and packaging for dedicated server" />
 	<Option Name="ExtraPackageOnlyArguments" DefaultValue="" Restrict="" Description="Extra arguments to use when staging only" />	
 	<Option Name="TargetPlatformFilteredExtraStageAndPackageArguments" DefaultValue="" Restrict="" Description="Extra arguments to use when staging and packaging.  Filtered to specific platforms via TargetPlatformFilteredExtraStageAndPackageArgumentsFilter." />
 	<Option Name="TargetPlatformFilteredExtraStageAndPackageArgumentsFilter" DefaultValue="" Restrict="" Description="Platform filter applied to TargetPlatformFilteredExtraStageAndPackageArguments." />
@@ -102,6 +103,7 @@
 
 	<!-- Does this project require a dedicated server for testing? (not currently implemented!)-->
 	<Option Name="DedicatedServer" DefaultValue="false" Description="Project uses a dedicated server" />
+	<Option Name="ServerPackagedSeparately" DefaultValue="false" Description="Whether dedicated server is packaged separately" />
 
 	<!-- Customizations -->
 	<Option Name="Versioned" Restrict="true|false" DefaultValue="false" Description="Whether to embed changelist number into binaries"/>
@@ -616,13 +618,17 @@
 		<Property Name="CloudArtifactPlatformPackagedNodeName" Value="$(PlatformPublishPackagedNodeName) (Cloud)" />
 		<Property Name="CloudArtifactPlatformRetrieve" Value="$(PreNodeNameWithPlatform)Retrieve Build Artifact $(TargetPlatform)" />
 		<Property Name="RequirePlatformStaged" Value="$(PlatformStageNodeName)" />
+		<Property Name="RequirePlatformStaged" Value="$(PlatformStageNodeName) Server" If="$(DedicatedServer)"/>
 		<Property Name="RequirePlatformStaged" Value="$(RequirePlatformStaged);$(PlatformStagedArtifactNodeName)" If="!$(SkipPublish) And $(PublishHordeArtifacts)"/>
 		<Property Name="RequirePlatformStaged" Value="$(RequirePlatformStaged);$(CloudPlatformStagedArtifactNodeName)" If="!$(SkipPublish) And $(PublishCloudArtifacts)"/>
 		<Property Name="RequirePlatformStaged" Value="$(RequirePlatformStaged);$(PlatformPublishStagedNodeName)" If="!$(SkipPublish) And $(PublishBuilds)"/>
-		<Property Name="RequirePlatformPackaged" Value="$(PlatformPackageNodeName)" />
+		<Property Name="RequirePlatformPackaged" Value="$(PlatformPackageNodeName)"/>
+		<Property Name="RequirePlatformPackaged" Value="$(PlatformPackageNodeName) Server" If="$(DedicatedServer)"/>
 		<Property Name="RequirePlatformPackaged" Value="$(RequirePlatformPackaged);$(HordeArtifactPlatformPackagedNodeName)" If="!$(SkipPublish) And $(PublishHordeArtifacts)"/>
+		<Property Name="RequirePlatformPackaged" Value="$(RequirePlatformPackaged);$(HordeArtifactPlatformPackagedNodeName) Server" If="!$(SkipPublish) And $(PublishHordeArtifacts) And $(DedicatedServer)"/>
 		<Property Name="RequirePlatformPackaged" Value="$(RequirePlatformPackaged);$(CloudArtifactPlatformPackagedNodeName)" If="!$(SkipPublish) And $(PublishCloudArtifacts)"/>
 		<Property Name="RequirePlatformPackaged" Value="$(RequirePlatformPackaged);$(PlatformPublishPackagedNodeName)"  If="!$(SkipPublish) And $(PublishBuilds)"/>
+		<Property Name="RequirePlatformPackaged" Value="$(RequirePlatformPackaged);$(PlatformPublishPackagedNodeName) Server"  If="!$(SkipPublish) And $(PublishBuilds) And $(DedicatedServer)"/>
 		
 		<!-- Set CookPlatform since it isn't always the same as the target -->
 		<Property Name="CookPlatform" Value="$(TargetPlatform)" />
@@ -650,6 +656,8 @@
 		<Property Name="StageAndPackageArguments" Value="$(StageAndPackageArguments) $(TargetPlatformFilteredExtraStageAndPackageArguments)" If="'$(TargetPlatformFilteredExtraStageAndPackageArgumentsFilter)' != '' and ContainsItem('$(TargetPlatformFilteredExtraStageAndPackageArgumentsFilter)', '$(TargetPlatform)', '+')" />
 		<Property Name="PackageOnlyArguments" Value="$(ExtraPackageOnlyArguments)" />
 
+		<Property Name="StageAndPackageServerArguments" Value="$(ExtraStageAndPackageServerArguments)" />
+
 		<Property Name="BuildSizeLimit" Value=""/>
 		<Regex Pattern="(?&lt;!\w)$(TargetPlatform):([0-9]+)" Capture="BuildSizeLimit" Input="$(CheckBuildSizePlatforms)"/>
 
@@ -669,16 +677,21 @@
 					<Default>
 						<ForEach Name="TargetConfiguration" Values="$(TargetConfigurations)" Separator="+">
 							<Compile Target="$(TargetName)" Platform="$(TargetPlatform)" Configuration="$(TargetConfiguration)" Arguments="-BuildVersion=&quot;$(BuildVersion)&quot; $(GenericCompileArguments) $(TargetCompileArguments)"/>
-							<Do If="'$(DedicatedServer)'" >
-								<Property Name="ServerPlatform" Value="Win64"/>
-								<Property Name="ServerPlatform" Value="Linux" If="$(IsTargetPlatformLinux)"/>
-								<Compile Target="$(TargetName)Server" Platform="$(ServerPlatform)" Configuration="$(TargetConfiguration)" Arguments="-BuildVersion=&quot;$(BuildVersion)&quot; $(GenericCompileArguments) $(TargetCompileArguments)"/>
-							</Do>
 						</ForEach>
 					</Default>
 				</Switch>
 			</Node>
 
+			<Node Name="$(PlatformCompileNodeName) Server" Requires="$(PlatformToolsNodeName)">
+				<ForEach Name="TargetConfiguration" Values="$(TargetConfigurations)" Separator="+">
+					<Property Name="ServerPlatform" Value="Win64"/>
+					<Property Name="ServerPlatform" Value="Linux" If="$(IsTargetPlatformLinux)"/>
+					<Property Name="ServerCompileArguments" Value="-BuildVersion=&quot;$(BuildVersion)&quot; $(GenericCompileArguments) $(TargetCompileArguments)"/>
+					<Property Name="ServerCompileArguments" Value="$(ServerCompileArguments) -BinariesSubfolder=$(TargetName)Server" If="$(ServerPackagedSeparately)"/>
+					<Compile Target="$(TargetName)Server" Platform="$(ServerPlatform)" Configuration="$(TargetConfiguration)" Arguments="$(ServerCompileArguments)"/>
+				</ForEach>	
+			</Node>
+
 			<!-- (Optional) Upload Symbols -->
 			<Do If="'$(UploadSymbolsNodeName)' != ''">
 				<Node Name="$(UploadSymbolsNodeName)" Requires="$(PlatformCompileNodeName)">
@@ -709,6 +722,20 @@
 				</Do>
 			</Node>
 			
+			<Node Name="$(PlatformCookNodeName) Server" Requires="$(PlatformEditorNodeName)" Produces="#$(PlatformCookNodeName) Server Complete" RunEarly="true">
+				<Do If="!$(SkipCook)">
+					<Expand Name="PreCookStep"/>
+					<Property Name="CookArgs" Value="" />
+					<Property Name="CookArgs" Value="-MapIniSection=&quot;$(MapIniSection)&quot; $(CookPlatform)" If="'$(MapIniSection)' != ''"/>	
+
+					<Property Name="ServerPlatformCook" Value="Windows"/>
+					<Property Name="ServerPlatformCook" Value="Linux" If="$(IsTargetPlatformLinux)"/>
+					<Property Name="ServerPlatformCook" Value="$(ServerPlatformCook)Server"/>
+
+					<Cook Project="$(TargetProject)" Platform="$(ServerPlatformCook)" TagOutput="false" Arguments="$(CookArgs) $(ExtraCookArguments)"/>
+				</Do>
+			</Node>
+			
 			<!-- Stage -->
 			<Node Name="$(PlatformStageNodeName)" Requires="$(PlatformCompileNodeName);#$(PlatformCookNodeName) Complete">
 				<Switch>
@@ -722,6 +749,12 @@
 					</Default>
 				</Switch>
 			</Node>
+
+			<Node Name="$(PlatformStageNodeName) Server" Requires="$(PlatformCompileNodeName) Server;#$(PlatformCookNodeName) Server Complete">
+				<Property Name="ServerStageArguments" Value="$(BCRArgs) -configuration=$(PackageConfigurations) $(StageAndPackageServerArguments) -skipbuild -skipcook -stage -pak -NoCodeSign" />
+				<Property Name="ServerStageArguments" Value="$(ServerStageArguments) -StagingDirectory=&quot;$(ProjectOutputDirectory)/$(StagedPlatformFolder)/Server&quot;" If="$(ServerPackagedSeparately)" />
+				<Command Name="BuildCookRun" Arguments="$(ServerStageArguments)" />
+			</Node>
 			
 			<!-- Hash -->
 			<Node Name="$(PlatformHashNodeName)" Requires="$(PlatformStageNodeName)">
@@ -790,6 +823,15 @@
 					<Copy Files="..." From="$(ShaderSymbolsDirectory)/*$(TargetPlatform)*" To="$(NetworkOutputDirectory)/$(StagedPlatformFolder)/ShaderSymbols"/>
 				</Do>
 			</Node>
+			
+			<Node Name="$(PlatformPublishStagedNodeName) Server" Requires="$(PlatformPublishStageDependencyName)">
+				<Log Message="Network upload is disabled, skipping upload." If="!$(PublishToNetworkOutputDirectory)"/>
+				<Do If="$(PublishToNetworkOutputDirectory)">				
+					<Error Message="NetworkOutputDirectory is not set" If="'$(NetworkOutputDirectory)' == ''"/>
+					<Copy Files="..." From="$(ProjectOutputDirectory)/$(StagedPlatformFolder)/Server" To="$(NetworkOutputDirectory)/$(StagedPlatformFolder)/Server/Staged"/>
+					<Copy Files="..." From="$(ShaderSymbolsDirectory)/*$(TargetPlatform)Server*" To="$(NetworkOutputDirectory)/$(StagedPlatformFolder)/Server/ShaderSymbols"/>
+				</Do>
+			</Node>
 
 			<!-- Publish (Staged) (Cloud) -->
 			<Do If="$(PublishCloudArtifacts)">
@@ -807,7 +849,7 @@
 			<Do If="$(PublishHordeArtifacts)">
 				<Node Name="$(PlatformStagedArtifactNodeName)" Requires="$(PlatformStageNodeName)">
 					<Expand Name="WriteArtifactMetadataFile" Type="staged-build" Platform="$(StagedPlatformFolder)" Configuration="$(PackageConfigurations)" ArtifactDir="$(ProjectOutputDirectory)/$(StagedPlatformFolder)"/>
-					<Tag Files="$(ProjectOutputDirectory)/$(StagedPlatformFolder)/..." With="#$(PlatformStageNodeName) Files"/>
+					<Tag Files="$(ProjectOutputDirectory)/$(StagedPlatformFolder)/..." Except=".../Server/..." With="#$(PlatformStageNodeName) Files"/>
 
 					<CreateArtifact Name="$(ProjectName)-Staged-$(StagedPlatformFolder)" Type="staged-build" BaseDir="$(ProjectOutputDirectory)/$(StagedPlatformFolder)" Keys="Project=$(ProjectName)" Metadata="Stream=$(Branch);Platform=$(StagedPlatformFolder);BuildName=$(BuildNamePath);dashboard-category=Staged Builds" Files="#$(PlatformStageNodeName) Files"/>
 				</Node>
@@ -820,6 +862,14 @@
 				</ForEach>
 			</Node>
 			
+			<Node Name="$(PlatformPackageNodeName) Server" Requires="$(PlatformPublishStageDependencyName)">
+				<ForEach Name="PackageConfiguration" Values="$(PackageConfigurations)" Separator="+">
+					<Property Name="ServerArchiveDirectory" Value="$(ProjectOutputDirectory)/Packages/$(StagedPlatformFolder)/$(PackageConfiguration)" If="!$(ServerPackagedSeparately)" />
+					<Property Name="ServerArchiveDirectory" Value="$(ProjectOutputDirectory)/Server/Packages/$(StagedPlatformFolder)/$(PackageConfiguration)" If="$(ServerPackagedSeparately)" />
+					<Command Name="BuildCookRun" Arguments="$(BCRArgs) -configuration=$(PackageConfiguration) $(StageAndPackageServerArguments) -skipbuild -skipcook -skipstage -skippak -package $(PackageOnlyArguments) -archive -archivedirectory=&quot;$(ServerArchiveDirectory)&quot;" />
+				</ForEach>
+			</Node>
+			
 			<!-- Publish (Packages) -->
 			<Node Name="$(PlatformPublishPackagedNodeName)" Requires="$(PlatformPackageNodeName)">
 				<Log Message="Network upload is disabled, skipping upload." If="!$(PublishToNetworkOutputDirectory)"/>
@@ -830,6 +880,16 @@
 					</ForEach>
 				</Do>
 			</Node>
+			
+			<Node Name="$(PlatformPublishPackagedNodeName) Server" Requires="$(PlatformPackageNodeName)">
+				<Log Message="Network upload is disabled, skipping upload." If="!$(PublishToNetworkOutputDirectory)"/>
+				<Do If="$(PublishToNetworkOutputDirectory)">
+					<Error Message="NetworkOutputDirectory is not set" If="'$(NetworkOutputDirectory)' == ''"/>
+					<ForEach Name="PackageConfiguration" Values="$(PackageConfigurations)" Separator="+">
+						<Command Name="BuildCookRun" Arguments="$(BCRArgs) -configuration=$(PackageConfiguration) $(StageAndPackageServerArguments) -skipbuild -skipcook -skipstage -skippak -skippackage -archive -archivedirectory=&quot;$(NetworkOutputDirectory)/$(StagedPlatformFolder)/Server/Packages&quot;" />
+					</ForEach>
+				</Do>
+			</Node>
 
 			<!-- Publish (Packages) (Cloud) -->
 			<Do  If="$(PublishCloudArtifacts)">
@@ -854,11 +914,20 @@
 					<ForEach Name="PackageConfiguration" Values="$(PackageConfigurations)" Separator="+">
 						<Command Name="BuildCookRun" Arguments="$(BCRArgs) -configuration=$(PackageConfiguration) $(StageAndPackageArguments) -skipbuild -skipcook -skippak $(HordePackageSkipBCRArguments) -archive -archivedirectory=&quot;$(ProjectOutputDirectory)/Packages/$(StagedPlatformFolder)/$(PackageConfiguration)&quot;" />
 						<Expand Name="WriteArtifactMetadataFile" Type="packaged-build" Platform="$(StagedPlatformFolder)" Configuration="$(PackageConfiguration)" ArtifactDir="$(ProjectOutputDirectory)/Packages/$(StagedPlatformFolder)/$(PackageConfiguration)"/>
-						<Tag Files="$(ProjectOutputDirectory)/Packages/$(StagedPlatformFolder)/$(PackageConfiguration)/..." With="#$(PlatformPackageNodeName) $(PackageConfiguration) Files"/>
-
+						<Tag Files="$(ProjectOutputDirectory)/Packages/$(StagedPlatformFolder)/$(PackageConfiguration)/..." Except=".../Server/..." With="#$(PlatformPackageNodeName) $(PackageConfiguration) Files"/>
 						<CreateArtifact Name="$(ProjectName)-Packaged-$(StagedPlatformFolder)-$(PackageConfiguration)" Type="packaged-build" BaseDir="$(ProjectOutputDirectory)/Packages/$(StagedPlatformFolder)/$(PackageConfiguration)" Keys="Project=$(ProjectName)" Metadata="Stream=$(Branch);Platform=$(StagedPlatformFolder);Configuration=$(PackageConfiguration);BuildName=$(BuildNamePath);dashboard-category=$(PackageConfiguration) Builds" Files="#$(PlatformPackageNodeName) $(PackageConfiguration) Files"/>
 					</ForEach>
 				</Node>
+			
+				<Node Name="$(HordeArtifactPlatformPackagedNodeName) Server" Requires="$(PlatformPackageNodeName)">
+					<ForEach Name="PackageConfiguration" Values="$(PackageConfigurations)" Separator="+">
+						<Property Name="ServerArchivePath" Value="$(ProjectOutputDirectory)/Server/Packages/$(StagedPlatformFolder)/$(PackageConfiguration)" If="$(ServerPackagedSeparately)" />
+						<Property Name="ServerArchivePath" Value="$(ProjectOutputDirectory)/Packages/$(StagedPlatformFolder)/$(PackageConfiguration)" If="!$(ServerPackagedSeparately)" />
+						<Command Name="BuildCookRun" Arguments="$(BCRArgs) -configuration=$(PackageConfiguration) $(StageAndPackageServerArguments) -skipbuild -skipcook -skippak $(HordePackageSkipBCRArguments) -archive -archivedirectory=&quot;$(ServerArchivePath)&quot;" />
+						<Tag Files="$(ServerArchivePath)/..." With="#$(PlatformPackageNodeName) $(PackageConfiguration) Server Files"/>
+						<CreateArtifact Name="$(ProjectName)-Server-Packaged-$(StagedPlatformFolder)-$(PackageConfiguration)" Type="packaged-build" BaseDir="$(ServerArchivePath)" Keys="Project=$(ProjectName)" Metadata="Stream=$(Branch);Platform=$(StagedPlatformFolder);Configuration=$(PackageConfiguration);BuildName=$(BuildNamePath);dashboard-category=$(PackageConfiguration) Builds" Files="#$(PlatformPackageNodeName) $(PackageConfiguration) Server Files"/>
+					</ForEach>
+				</Node>
 			</Do>
 
 			<!-- Check Build Size -->
@@ -1048,6 +1117,9 @@
 		<Label Category="Clients" Name="$(PreLabelName)$(TargetPlatform)" Requires="$(PlatformCompileNodeName);$(UploadSymbolsNodeName)" />
 		<Label Category="Clients" Name="$(PreLabelName)Content $(TargetPlatform)" Requires="$(PlatformCookNodeName);$(PlatformStageNodeName);$(PlatformPackageNodeName)" />
 		<Label Category="Clients" Name="$(PreLabelName)Publish $(TargetPlatform)" Requires="$(PlatformPublishPackagedNodeName);$(PlatformPublishPackagedNodeName);$(PlatformPublishStagedNodeName);$(PlatformHashNodeName)" />
+		<Label Category="Servers" Name="$(PreLabelName)$(TargetPlatform) Server" Requires="$(PlatformCompileNodeName) Server" />
+		<Label Category="Servers" Name="$(PreLabelName)Content $(TargetPlatform) Server" Requires="$(PlatformCookNodeName) Server;$(PlatformStageNodeName) Server;$(PlatformPackageNodeName) Server" />
+		<Label Category="Servers" Name="$(PreLabelName)Publish $(TargetPlatform) Server " Requires="$(PlatformPublishPackagedNodeName) Server;$(PlatformPublishStagedNodeName) Server" />
 		<Label Category="Tests" Name="$(PreLabelName)Install $(TargetPlatform)" Requires="$(AllPlatformInstallNodes)" />
 		<Label Category="Tests" Name="$(PreLabelName)BootTest $(TargetPlatform)" Requires="$(AllPlatformBootTestNodes)" />
 		<Label Category="Tests" Name="$(PreLabelName)Test $(TargetPlatform)" Requires="$(AllPlatformTestNodes)" />
@@ -1055,7 +1127,9 @@
 		<!-- Tie all steps into our platform tasks unless excluded -->
 		<Property Name="PlatformTasks" Value="" />
 		<Property Name="PlatformTasks" Value="$(PlatformCompileNodeName);$(UploadSymbolsNodeName)" If="!$(SkipBuild)" />
+		<Property Name="PlatformTasks" Value="$(PlatformCompileNodeName) Server" If="!$(SkipBuild) and $(DedicatedServer)" />
 		<Property Name="PlatformTasks" Value="$(PlatformTasks);$(PlatformCookNodeName)" If="!$(SkipCook) and !$(SkipPublish)" />
+		<Property Name="PlatformTasks" Value="$(PlatformTasks);$(PlatformCookNodeName) Server" If="!$(SkipCook) and !$(SkipPublish) and $(DedicatedServer)" />
 		<Property Name="PlatformTasks" Value="$(PlatformTasks);$(RequirePlatformStaged)" If="!$(SkipPublish)" />
 		<Property Name="PlatformTasks" Value="$(PlatformTasks);$(RequirePlatformPackaged)" If="!$(SkipPackage) and !$(SkipPublish)" />
 		<Property Name="PlatformTasks" Value="$(PlatformTasks);$(PlatformHashNodeName)"  If="$(PublishHashedStagingData)" />	
diff --git a/Engine/Source/Programs/UnrealBuildTool/Configuration/TargetRules.cs b/Engine/Source/Programs/UnrealBuildTool/Configuration/TargetRules.cs
--- a/Engine/Source/Programs/UnrealBuildTool/Configuration/TargetRules.cs
+++ b/Engine/Source/Programs/UnrealBuildTool/Configuration/TargetRules.cs
@@ -895,6 +895,7 @@ namespace UnrealBuildTool
 		/// Subfolder to place executables in, relative to the default location.
 		/// </summary>
 		[RequiresUniqueBuildEnvironment]
+		[CommandLine("-BinariesSubfolder")]
 		public string ExeBinariesSubFolder { get; set; } = String.Empty;
 
 		/// <summary>
