Separate client and server packages on demand - BTP.xml

#rnx

diff --git a/Engine/Build/Graph/Tasks/BuildAndTestProject.xml b/Engine/Build/Graph/Tasks/BuildAndTestProject.xml
--- a/Engine/Build/Graph/Tasks/BuildAndTestProject.xml
+++ b/Engine/Build/Graph/Tasks/BuildAndTestProject.xml
@@ -620,8 +620,10 @@
 		<Property Name="RequirePlatformStaged" Value="$(PlatformStageNodeName)" />
 		<Property Name="RequirePlatformStaged" Value="$(PlatformStageNodeName) Server" If="$(DedicatedServer)"/>
 		<Property Name="RequirePlatformStaged" Value="$(RequirePlatformStaged);$(PlatformStagedArtifactNodeName)" If="!$(SkipPublish) And $(PublishHordeArtifacts)"/>
+		<Property Name="RequirePlatformStaged" Value="$(RequirePlatformStaged);$(PlatformStagedArtifactNodeName) Server" If="!$(SkipPublish) And $(PublishHordeArtifacts) And $(DedicatedServer)"/>
 		<Property Name="RequirePlatformStaged" Value="$(RequirePlatformStaged);$(CloudPlatformStagedArtifactNodeName)" If="!$(SkipPublish) And $(PublishCloudArtifacts)"/>
 		<Property Name="RequirePlatformStaged" Value="$(RequirePlatformStaged);$(PlatformPublishStagedNodeName)" If="!$(SkipPublish) And $(PublishBuilds)"/>
+		<Property Name="RequirePlatformStaged" Value="$(RequirePlatformStaged);$(PlatformPublishStagedNodeName) Server" If="!$(SkipPublish) And $(PublishBuilds) And $(DedicatedServer)"/>
 		<Property Name="RequirePlatformPackaged" Value="$(PlatformPackageNodeName)"/>
 		<Property Name="RequirePlatformPackaged" Value="$(PlatformPackageNodeName) Server" If="$(DedicatedServer)"/>
 		<Property Name="RequirePlatformPackaged" Value="$(RequirePlatformPackaged);$(HordeArtifactPlatformPackagedNodeName)" If="!$(SkipPublish) And $(PublishHordeArtifacts)"/>
@@ -649,7 +651,6 @@
 		<!-- Default BuildCookRun args needed when staging/packaging/archiving this configuration -->
 		<Property Name="BCRArgs" Value="-project=$(TargetProject) -platform=$(TargetPlatform) -cookflavor=&quot;$(BCRCookFlavor)&quot;" />
 		<Property Name="BCRArgs" Value="$(BCRArgs) -client" If="'$(TargetType)' == 'Client'" />
-		<Property Name="BCRArgs" Value="$(BCRArgs) -server" If="'$(DedicatedServer)'" />
 
 		<!-- Support additional stage and package arguments that are only active on some platforms -->
 		<Property Name="StageAndPackageArguments" Value="$(ExtraStageAndPackageArguments)" />
@@ -687,7 +688,6 @@
 					<Property Name="ServerPlatform" Value="Win64"/>
 					<Property Name="ServerPlatform" Value="Linux" If="$(IsTargetPlatformLinux)"/>
 					<Property Name="ServerCompileArguments" Value="-BuildVersion=&quot;$(BuildVersion)&quot; $(GenericCompileArguments) $(TargetCompileArguments)"/>
-					<Property Name="ServerCompileArguments" Value="$(ServerCompileArguments) -BinariesSubfolder=$(TargetName)Server" If="$(ServerPackagedSeparately)"/>
 					<Compile Target="$(TargetName)Server" Platform="$(ServerPlatform)" Configuration="$(TargetConfiguration)" Arguments="$(ServerCompileArguments)"/>
 				</ForEach>	
 			</Node>
@@ -745,14 +745,15 @@
 						</ForEach>
 					</Case>
 					<Default>
-						<Command Name="BuildCookRun" Arguments="$(BCRArgs) -configuration=$(PackageConfigurations) $(StageAndPackageArguments) -skipbuild -skipcook -stage -pak -NoCodeSign" />
+						<Property Name="ClientStageArguments" Value="$(BCRArgs) -configuration=$(PackageConfigurations) $(StageAndPackageArguments) -skipbuild -skipcook -stage -pak -NoCodeSign" />
+						<Command Name="BuildCookRun" Arguments="$(ClientStageArguments)" />
 					</Default>
 				</Switch>
 			</Node>
 
 			<Node Name="$(PlatformStageNodeName) Server" Requires="$(PlatformCompileNodeName) Server;#$(PlatformCookNodeName) Server Complete">
-				<Property Name="ServerStageArguments" Value="$(BCRArgs) -configuration=$(PackageConfigurations) $(StageAndPackageServerArguments) -skipbuild -skipcook -stage -pak -NoCodeSign" />
-				<Property Name="ServerStageArguments" Value="$(ServerStageArguments) -StagingDirectory=&quot;$(ProjectOutputDirectory)/$(StagedPlatformFolder)/Server&quot;" If="$(ServerPackagedSeparately)" />
+				<Property Name="BCRArgsServer" Value="$(BCRArgs) -server" />
+				<Property Name="ServerStageArguments" Value="$(BCRArgsServer) -configuration=$(PackageConfigurations) $(StageAndPackageServerArguments) -skipbuild -skipcook -stage -pak -NoCodeSign" />
 				<Command Name="BuildCookRun" Arguments="$(ServerStageArguments)" />
 			</Node>
 			
@@ -828,8 +829,8 @@
 				<Log Message="Network upload is disabled, skipping upload." If="!$(PublishToNetworkOutputDirectory)"/>
 				<Do If="$(PublishToNetworkOutputDirectory)">				
 					<Error Message="NetworkOutputDirectory is not set" If="'$(NetworkOutputDirectory)' == ''"/>
-					<Copy Files="..." From="$(ProjectOutputDirectory)/$(StagedPlatformFolder)/Server" To="$(NetworkOutputDirectory)/$(StagedPlatformFolder)/Server/Staged"/>
-					<Copy Files="..." From="$(ShaderSymbolsDirectory)/*$(TargetPlatform)Server*" To="$(NetworkOutputDirectory)/$(StagedPlatformFolder)/Server/ShaderSymbols"/>
+					<Copy Files="..." From="$(ProjectOutputDirectory)/$(StagedPlatformFolder)Server" To="$(NetworkOutputDirectory)/$(StagedPlatformFolder)Server/Staged"/>
+					<Copy Files="..." From="$(ShaderSymbolsDirectory)/*$(TargetPlatform)Server*" To="$(NetworkOutputDirectory)/$(StagedPlatformFolder)Server/ShaderSymbols"/>
 				</Do>
 			</Node>
 
@@ -848,25 +849,41 @@
 			<!-- Publish (Staged) (Horde) -->
 			<Do If="$(PublishHordeArtifacts)">
 				<Node Name="$(PlatformStagedArtifactNodeName)" Requires="$(PlatformStageNodeName)">
-					<Expand Name="WriteArtifactMetadataFile" Type="staged-build" Platform="$(StagedPlatformFolder)" Configuration="$(PackageConfigurations)" ArtifactDir="$(ProjectOutputDirectory)/$(StagedPlatformFolder)"/>
-					<Tag Files="$(ProjectOutputDirectory)/$(StagedPlatformFolder)/..." Except=".../Server/..." With="#$(PlatformStageNodeName) Files"/>
+					<Property Name="ArtifactsPath" Value="$(ProjectOutputDirectory)/$(StagedPlatformFolder)" />
+					
+					<Expand Name="WriteArtifactMetadataFile" Type="staged-build" Platform="$(StagedPlatformFolder)" Configuration="$(PackageConfigurations)" ArtifactDir="$(ArtifactsPath)"/>
+					<Tag Files="$(ArtifactsPath)/..." With="#$(PlatformStageNodeName) Files"/>
 
-					<CreateArtifact Name="$(ProjectName)-Staged-$(StagedPlatformFolder)" Type="staged-build" BaseDir="$(ProjectOutputDirectory)/$(StagedPlatformFolder)" Keys="Project=$(ProjectName)" Metadata="Stream=$(Branch);Platform=$(StagedPlatformFolder);BuildName=$(BuildNamePath);dashboard-category=Staged Builds" Files="#$(PlatformStageNodeName) Files"/>
+					<CreateArtifact Name="$(ProjectName)-Staged-$(StagedPlatformFolder)" Type="staged-build" BaseDir="$(ArtifactsPath)" Keys="Project=$(ProjectName)" Metadata="Stream=$(Branch);Platform=$(StagedPlatformFolder);BuildName=$(BuildNamePath);dashboard-category=Staged Builds" Files="#$(PlatformStageNodeName) Files"/>
+				</Node>
+			</Do>
+			
+			<Do If="$(PublishHordeArtifacts)">
+				<Node Name="$(PlatformStagedArtifactNodeName) Server" Requires="$(PlatformStageNodeName)">
+					<Property Name="ServerArtifactsPath" Value="$(ProjectOutputDirectory)/$(StagedPlatformFolder)" />
+					<Property Name="ServerArtifactsPath" Value="$(ProjectOutputDirectory)/$(StagedPlatformFolder)Server" If="$(ServerPackagedSeparately)" />
+					
+					<Expand Name="WriteArtifactMetadataFile" Type="staged-build" Platform="$(StagedPlatformFolder)" Configuration="$(PackageConfigurations)" ArtifactDir="$(ServerArtifactsPath)"/>
+					<Tag Files="$(ServerArtifactsPath)/..." Except="$(TargetPlatform)Server/..." With="#$(PlatformStageNodeName) Server Files"/>
+
+					<CreateArtifact Name="$(ProjectName)-Server-Staged-$(StagedPlatformFolder)" Type="staged-build" BaseDir="$(ServerArtifactsPath)" Keys="Project=$(ProjectName)" Metadata="Stream=$(Branch);Platform=$(StagedPlatformFolder);BuildName=$(BuildNamePath);dashboard-category=Staged Builds" Files="#$(PlatformStageNodeName) Server Files"/>
 				</Node>
 			</Do>
 
 			<!-- Package -->
 			<Node Name="$(PlatformPackageNodeName)" Requires="$(PlatformPublishStageDependencyName)">
 				<ForEach Name="PackageConfiguration" Values="$(PackageConfigurations)" Separator="+">
-					<Command Name="BuildCookRun" Arguments="$(BCRArgs) -configuration=$(PackageConfiguration) $(StageAndPackageArguments) -skipbuild -skipcook -skipstage -skippak -package $(PackageOnlyArguments) -archive -archivedirectory=&quot;$(ProjectOutputDirectory)/Packages/$(StagedPlatformFolder)/$(PackageConfiguration)&quot;" />
+					<Property Name="ArchiveDirectory" Value="$(ProjectOutputDirectory)/Packages/$(StagedPlatformFolder)/$(PackageConfiguration)" />
+					<Command Name="BuildCookRun" Arguments="$(BCRArgs) -configuration=$(PackageConfiguration) $(StageAndPackageArguments) -skipbuild -skipcook -skipstage -skippak -package $(PackageOnlyArguments) -archive -archivedirectory=&quot;$(ArchiveDirectory)&quot;" />
 				</ForEach>
 			</Node>
 			
 			<Node Name="$(PlatformPackageNodeName) Server" Requires="$(PlatformPublishStageDependencyName)">
 				<ForEach Name="PackageConfiguration" Values="$(PackageConfigurations)" Separator="+">
-					<Property Name="ServerArchiveDirectory" Value="$(ProjectOutputDirectory)/Packages/$(StagedPlatformFolder)/$(PackageConfiguration)" If="!$(ServerPackagedSeparately)" />
-					<Property Name="ServerArchiveDirectory" Value="$(ProjectOutputDirectory)/Server/Packages/$(StagedPlatformFolder)/$(PackageConfiguration)" If="$(ServerPackagedSeparately)" />
-					<Command Name="BuildCookRun" Arguments="$(BCRArgs) -configuration=$(PackageConfiguration) $(StageAndPackageServerArguments) -skipbuild -skipcook -skipstage -skippak -package $(PackageOnlyArguments) -archive -archivedirectory=&quot;$(ServerArchiveDirectory)&quot;" />
+					<Property Name="ServerArchiveDirectory" Value="$(ProjectOutputDirectory)/Packages/$(StagedPlatformFolder)/$(PackageConfiguration)" />
+					<Property Name="ServerArchiveDirectory" Value="$(ProjectOutputDirectory)/Packages/$(StagedPlatformFolder)Server/$(PackageConfiguration)" If="$(ServerPackagedSeparately)" />
+					<Property Name="BCRArgsServer" Value="$(BCRArgs) -server" />	
+				<Command Name="BuildCookRun" Arguments="$(BCRArgsServer) -configuration=$(PackageConfiguration) $(StageAndPackageServerArguments) -skipbuild -skipcook -skipstage -skippak -package $(PackageOnlyArguments) -archive -archivedirectory=&quot;$(ServerArchiveDirectory)&quot;" />
 				</ForEach>
 			</Node>
 			
@@ -876,7 +893,8 @@
 				<Do If="$(PublishToNetworkOutputDirectory)">
 					<Error Message="NetworkOutputDirectory is not set" If="'$(NetworkOutputDirectory)' == ''"/>
 					<ForEach Name="PackageConfiguration" Values="$(PackageConfigurations)" Separator="+">
-						<Command Name="BuildCookRun" Arguments="$(BCRArgs) -configuration=$(PackageConfiguration) $(StageAndPackageArguments) -skipbuild -skipcook -skipstage -skippak -skippackage -archive -archivedirectory=&quot;$(NetworkOutputDirectory)/$(StagedPlatformFolder)/Packages&quot;" />
+						<Property Name="ArchiveDirectory" Value="$(NetworkOutputDirectory)/Packages/$(StagedPlatformFolder)"/>
+						<Command Name="BuildCookRun" Arguments="$(BCRArgs) -configuration=$(PackageConfiguration) $(StageAndPackageArguments) -skipbuild -skipcook -skipstage -skippak -skippackage -archive -archivedirectory=&quot;$(ArchiveDirectory)&quot;" />
 					</ForEach>
 				</Do>
 			</Node>
@@ -886,7 +904,10 @@
 				<Do If="$(PublishToNetworkOutputDirectory)">
 					<Error Message="NetworkOutputDirectory is not set" If="'$(NetworkOutputDirectory)' == ''"/>
 					<ForEach Name="PackageConfiguration" Values="$(PackageConfigurations)" Separator="+">
-						<Command Name="BuildCookRun" Arguments="$(BCRArgs) -configuration=$(PackageConfiguration) $(StageAndPackageServerArguments) -skipbuild -skipcook -skipstage -skippak -skippackage -archive -archivedirectory=&quot;$(NetworkOutputDirectory)/$(StagedPlatformFolder)/Server/Packages&quot;" />
+						<Property Name="ServerArchiveDirectory" Value="$(NetworkOutputDirectory)/Packages/$(StagedPlatformFolder)"/>
+						<Property Name="ServerArchiveDirectory" Value="$(NetworkOutputDirectory)/Packages/$(StagedPlatformFolder)Server" If="$(ServerPackagedSeparately)"/>
+						<Property Name="BCRArgsServer" Value="$(BCRArgs) -server" />	
+						<Command Name="BuildCookRun" Arguments="$(BCRArgsServer) -configuration=$(PackageConfiguration) $(StageAndPackageServerArguments) -skipbuild -skipcook -skipstage -skippak -skippackage -archive -archivedirectory=&quot;$(ServerArchiveDirectory)&quot;" />
 					</ForEach>
 				</Do>
 			</Node>
@@ -912,19 +933,21 @@
 
 				<Node Name="$(HordeArtifactPlatformPackagedNodeName)" Requires="$(PlatformPackageNodeName)">
 					<ForEach Name="PackageConfiguration" Values="$(PackageConfigurations)" Separator="+">
-						<Command Name="BuildCookRun" Arguments="$(BCRArgs) -configuration=$(PackageConfiguration) $(StageAndPackageArguments) -skipbuild -skipcook -skippak $(HordePackageSkipBCRArguments) -archive -archivedirectory=&quot;$(ProjectOutputDirectory)/Packages/$(StagedPlatformFolder)/$(PackageConfiguration)&quot;" />
-						<Expand Name="WriteArtifactMetadataFile" Type="packaged-build" Platform="$(StagedPlatformFolder)" Configuration="$(PackageConfiguration)" ArtifactDir="$(ProjectOutputDirectory)/Packages/$(StagedPlatformFolder)/$(PackageConfiguration)"/>
-						<Tag Files="$(ProjectOutputDirectory)/Packages/$(StagedPlatformFolder)/$(PackageConfiguration)/..." Except=".../Server/..." With="#$(PlatformPackageNodeName) $(PackageConfiguration) Files"/>
-						<CreateArtifact Name="$(ProjectName)-Packaged-$(StagedPlatformFolder)-$(PackageConfiguration)" Type="packaged-build" BaseDir="$(ProjectOutputDirectory)/Packages/$(StagedPlatformFolder)/$(PackageConfiguration)" Keys="Project=$(ProjectName)" Metadata="Stream=$(Branch);Platform=$(StagedPlatformFolder);Configuration=$(PackageConfiguration);BuildName=$(BuildNamePath);dashboard-category=$(PackageConfiguration) Builds" Files="#$(PlatformPackageNodeName) $(PackageConfiguration) Files"/>
+						<Property Name="ArchivePath" Value="$(ProjectOutputDirectory)/Packages/$(StagedPlatformFolder)/$(PackageConfiguration)" />
+						<Command Name="BuildCookRun" Arguments="$(BCRArgs) -configuration=$(PackageConfiguration) $(StageAndPackageArguments) -skipbuild -skipcook -skippak $(HordePackageSkipBCRArguments) -archive -archivedirectory=&quot;$(ArchivePath)&quot;" />
+						<Expand Name="WriteArtifactMetadataFile" Type="packaged-build" Platform="$(StagedPlatformFolder)" Configuration="$(PackageConfiguration)" ArtifactDir="$(ArchivePath)"/>
+						<Tag Files="$(ArchivePath)/..." Except="$(TargetPlatform)Server/...;*$(TargetPlatform)Server*" With="#$(PlatformPackageNodeName) $(PackageConfiguration) Files"/>
+						<CreateArtifact Name="$(ProjectName)-Packaged-$(StagedPlatformFolder)-$(PackageConfiguration)" Type="packaged-build" BaseDir="$(ArchivePath)" Keys="Project=$(ProjectName)" Metadata="Stream=$(Branch);Platform=$(StagedPlatformFolder);Configuration=$(PackageConfiguration);BuildName=$(BuildNamePath);dashboard-category=$(PackageConfiguration) Builds" Files="#$(PlatformPackageNodeName) $(PackageConfiguration) Files"/>
 					</ForEach>
 				</Node>
 			
 				<Node Name="$(HordeArtifactPlatformPackagedNodeName) Server" Requires="$(PlatformPackageNodeName)">
 					<ForEach Name="PackageConfiguration" Values="$(PackageConfigurations)" Separator="+">
-						<Property Name="ServerArchivePath" Value="$(ProjectOutputDirectory)/Server/Packages/$(StagedPlatformFolder)/$(PackageConfiguration)" If="$(ServerPackagedSeparately)" />
-						<Property Name="ServerArchivePath" Value="$(ProjectOutputDirectory)/Packages/$(StagedPlatformFolder)/$(PackageConfiguration)" If="!$(ServerPackagedSeparately)" />
-						<Command Name="BuildCookRun" Arguments="$(BCRArgs) -configuration=$(PackageConfiguration) $(StageAndPackageServerArguments) -skipbuild -skipcook -skippak $(HordePackageSkipBCRArguments) -archive -archivedirectory=&quot;$(ServerArchivePath)&quot;" />
-						<Tag Files="$(ServerArchivePath)/..." With="#$(PlatformPackageNodeName) $(PackageConfiguration) Server Files"/>
+						<Property Name="ServerArchivePath" Value="$(ProjectOutputDirectory)/Packages/$(StagedPlatformFolder)/$(PackageConfiguration)" />
+						<Property Name="ServerArchivePath" Value="$(ProjectOutputDirectory)/Packages/$(StagedPlatformFolder)Server/$(PackageConfiguration)" If="$(ServerPackagedSeparately)" />
+						<Property Name="BCRArgsServer" Value="$(BCRArgs) -server" />
+						<Command Name="BuildCookRun" Arguments="$(BCRArgsServer) -configuration=$(PackageConfiguration) $(StageAndPackageServerArguments) -skipbuild -skipcook -skippak $(HordePackageSkipBCRArguments) -archive -archivedirectory=&quot;$(ServerArchivePath)&quot;" />
+						<Tag Files="$(ServerArchivePath)/..."  Except="$(TargetPlatform)Server/..." With="#$(PlatformPackageNodeName) $(PackageConfiguration) Server Files"/>
 						<CreateArtifact Name="$(ProjectName)-Server-Packaged-$(StagedPlatformFolder)-$(PackageConfiguration)" Type="packaged-build" BaseDir="$(ServerArchivePath)" Keys="Project=$(ProjectName)" Metadata="Stream=$(Branch);Platform=$(StagedPlatformFolder);Configuration=$(PackageConfiguration);BuildName=$(BuildNamePath);dashboard-category=$(PackageConfiguration) Builds" Files="#$(PlatformPackageNodeName) $(PackageConfiguration) Server Files"/>
 					</ForEach>
 				</Node>
