Enable TestDataV2 in BuildAndTestProject

#rnx

diff --git a/Engine/Build/Graph/Tasks/Inc/GauntletSettings.xml b/Engine/Build/Graph/Tasks/Inc/GauntletSettings.xml
--- a/Engine/Build/Graph/Tasks/Inc/GauntletSettings.xml
+++ b/Engine/Build/Graph/Tasks/Inc/GauntletSettings.xml
@@ -31,7 +31,7 @@
 	<Option Name="DisableTestResultsForHorde" DefaultValue="false" Restrict="true|false" Description="Force disabling test results for Horde production"/>
 
 	<Property Name="GauntletCommonArgs" Value="-maxduration=$(GauntletTimeOut) -unattended -verbose -ResumeOnCriticalFailure $(ExtraGauntletCommonArgs)" />
-	<Property Name="GauntletCommonArgs" Value="$(GauntletCommonArgs) -WriteTestResultsForHorde" If="$(Horde) and !$(DisableTestResultsForHorde)"/>
+	<Property Name="GauntletCommonArgs" Value="$(GauntletCommonArgs) -WriteTestResultsForHorde -UseTestDataV2 -AddTestDataV1" If="$(Horde) and !$(DisableTestResultsForHorde)"/>
 	<Property Name="GauntletCommonArgs" Value="$(GauntletCommonArgs) -branch=&quot;$(Branch)&quot; -changelist=&quot;$(Change)&quot;" />
 
 	<Property Name="GauntletDeviceArgs" Value="" />
