Support for optional client compiles alongside regular game/server compiles
#rnx

diff --git a/Engine/Build/Graph/Tasks/BuildAndTestProject.xml b/Engine/Build/Graph/Tasks/BuildAndTestProject.xml
--- a/Engine/Build/Graph/Tasks/BuildAndTestProject.xml
+++ b/Engine/Build/Graph/Tasks/BuildAndTestProject.xml
@@ -30,8 +30,11 @@
 	<Option Name="IsContentOnlyProject" DefaultValue="false" Restrict="true|false" Description="Is the project a content only project?"/>	
 
 	<Option Name="TargetName" DefaultValue="" Restrict="" Description="If your project supports both Game and Client you can set this, otherwise one will be selected"/>
+	<Option Name="TargetNameClient" DefaultValue="" Restrict="" Description="Project target name used for a client"/>
+	<Option Name="TargetNameServer" DefaultValue="$(TargetName)Server" Restrict="" Description="Project target name used for the server"/>
 	<Option Name="EditorTargetName" DefaultValue="" Restrict="" Description="Specify the editor target name if different than (ProjectName)Editor"/>
 	<Option Name="CookedEditorTargetName" DefaultValue="" Restrict="" Description="Specify the cooked editor target name if different than (ProjectName)CookedEditor"/>
+	<Property Name="TargetName" Value="$(TargetNameClient)" If="$(TargetNameClient) != ''"/>
 
 	<Option Name="PreNodeName" DefaultValue="$(ProjectName)" Description="String to prepend to node name"/>
 	<Option Name="PreLabelName" DefaultValue="" Description="String to prepend to label name"/>
@@ -688,7 +691,7 @@
 					<Property Name="ServerPlatform" Value="Win64"/>
 					<Property Name="ServerPlatform" Value="Linux" If="$(IsTargetPlatformLinux)"/>
 					<Property Name="ServerCompileArguments" Value="-BuildVersion=&quot;$(BuildVersion)&quot; $(GenericCompileArguments) $(TargetCompileArguments)"/>
-					<Compile Target="$(TargetName)Server" Platform="$(ServerPlatform)" Configuration="$(TargetConfiguration)" Arguments="$(ServerCompileArguments)"/>
+					<Compile Target="$(TargetNameServer)" Platform="$(ServerPlatform)" Configuration="$(TargetConfiguration)" Arguments="$(ServerCompileArguments)"/>
 				</ForEach>	
 			</Node>
 
