Fixing bad token in BuildAndTestProject
#rnx

diff --git a/Engine/Build/Graph/Tasks/BuildAndTestProject.xml b/Engine/Build/Graph/Tasks/BuildAndTestProject.xml
--- a/Engine/Build/Graph/Tasks/BuildAndTestProject.xml
+++ b/Engine/Build/Graph/Tasks/BuildAndTestProject.xml
@@ -34,7 +34,7 @@
 	<Option Name="TargetNameServer" DefaultValue="$(TargetName)Server" Restrict="" Description="Project target name used for the server"/>
 	<Option Name="EditorTargetName" DefaultValue="" Restrict="" Description="Specify the editor target name if different than (ProjectName)Editor"/>
 	<Option Name="CookedEditorTargetName" DefaultValue="" Restrict="" Description="Specify the cooked editor target name if different than (ProjectName)CookedEditor"/>
-	<Property Name="TargetName" Value="$(TargetNameClient)" If="$(TargetNameClient) != ''"/>
+	<Property Name="TargetName" Value="$(TargetNameClient)" If="'$(TargetNameClient)' != ''"/>
 
 	<Option Name="PreNodeName" DefaultValue="$(ProjectName)" Description="String to prepend to node name"/>
 	<Option Name="PreLabelName" DefaultValue="" Description="String to prepend to label name"/>
