Ensure artifacts produced for server platforms don't contain 'client' in their name + metadata
#rnx

diff --git a/Engine/Build/Graph/Tasks/BuildAndTestProject.xml b/Engine/Build/Graph/Tasks/BuildAndTestProject.xml
--- a/Engine/Build/Graph/Tasks/BuildAndTestProject.xml
+++ b/Engine/Build/Graph/Tasks/BuildAndTestProject.xml
@@ -951,7 +951,7 @@
 						<Property Name="BCRArgsServer" Value="$(BCRArgs) -server" />
 						<Command Name="BuildCookRun" Arguments="$(BCRArgsServer) -configuration=$(PackageConfiguration) $(StageAndPackageServerArguments) -skipbuild -skipcook -skippak $(HordePackageSkipBCRArguments) -archive -archivedirectory=&quot;$(ServerArchivePath)&quot;" />
 						<Tag Files="$(ServerArchivePath)/..."  Except="$(TargetPlatform)Server/..." With="#$(PlatformPackageNodeName) $(PackageConfiguration) Server Files"/>
-						<CreateArtifact Name="$(ProjectName)-Server-Packaged-$(StagedPlatformFolder)-$(PackageConfiguration)" Type="packaged-build" BaseDir="$(ServerArchivePath)" Keys="Project=$(ProjectName)" Metadata="Stream=$(Branch);Platform=$(StagedPlatformFolder);Configuration=$(PackageConfiguration);BuildName=$(BuildNamePath);dashboard-category=$(PackageConfiguration) Builds" Files="#$(PlatformPackageNodeName) $(PackageConfiguration) Server Files"/>
+						<CreateArtifact Name="$(ProjectName)-Server-Packaged-$(TargetPlatform)-$(PackageConfiguration)" Type="packaged-build" BaseDir="$(ServerArchivePath)" Keys="Project=$(ProjectName)" Metadata="Stream=$(Branch);Platform=$(TargetPlatform);Configuration=$(PackageConfiguration);BuildName=$(BuildNamePath);dashboard-category=$(PackageConfiguration) Builds" Files="#$(PlatformPackageNodeName) $(PackageConfiguration) Server Files"/>
 					</ForEach>
 				</Node>
 			</Do>
