Fix network publish pathnames if the server is being packaged separately
#rnx

diff --git a/Engine/Build/Graph/Tasks/BuildAndTestProject.xml b/Engine/Build/Graph/Tasks/BuildAndTestProject.xml
--- a/Engine/Build/Graph/Tasks/BuildAndTestProject.xml
+++ b/Engine/Build/Graph/Tasks/BuildAndTestProject.xml
@@ -884,7 +884,7 @@
 			<Node Name="$(PlatformPackageNodeName) Server" Requires="$(PlatformPublishStageDependencyName)">
 				<ForEach Name="PackageConfiguration" Values="$(PackageConfigurations)" Separator="+">
 					<Property Name="ServerArchiveDirectory" Value="$(ProjectOutputDirectory)/Packages/$(StagedPlatformFolder)/$(PackageConfiguration)" />
-					<Property Name="ServerArchiveDirectory" Value="$(ProjectOutputDirectory)/Packages/$(StagedPlatformFolder)Server/$(PackageConfiguration)" If="$(ServerPackagedSeparately)" />
+					<Property Name="ServerArchiveDirectory" Value="$(ProjectOutputDirectory)/Packages/$(TargetPlatform)Server/$(PackageConfiguration)" If="$(ServerPackagedSeparately)" />
 					<Property Name="BCRArgsServer" Value="$(BCRArgs) -server" />	
 				<Command Name="BuildCookRun" Arguments="$(BCRArgsServer) -configuration=$(PackageConfiguration) $(StageAndPackageServerArguments) -skipbuild -skipcook -skipstage -skippak -package $(PackageOnlyArguments) -archive -archivedirectory=&quot;$(ServerArchiveDirectory)&quot;" />
 				</ForEach>
@@ -908,7 +908,7 @@
 					<Error Message="NetworkOutputDirectory is not set" If="'$(NetworkOutputDirectory)' == ''"/>
 					<ForEach Name="PackageConfiguration" Values="$(PackageConfigurations)" Separator="+">
 						<Property Name="ServerArchiveDirectory" Value="$(NetworkOutputDirectory)/Packages/$(StagedPlatformFolder)"/>
-						<Property Name="ServerArchiveDirectory" Value="$(NetworkOutputDirectory)/Packages/$(StagedPlatformFolder)Server" If="$(ServerPackagedSeparately)"/>
+						<Property Name="ServerArchiveDirectory" Value="$(NetworkOutputDirectory)/Packages/$(TargetPlatform)Server" If="$(ServerPackagedSeparately)"/>
 						<Property Name="BCRArgsServer" Value="$(BCRArgs) -server" />	
 						<Command Name="BuildCookRun" Arguments="$(BCRArgsServer) -configuration=$(PackageConfiguration) $(StageAndPackageServerArguments) -skipbuild -skipcook -skipstage -skippak -skippackage -archive -archivedirectory=&quot;$(ServerArchiveDirectory)&quot;" />
 					</ForEach>
