Updating Cloud builds for phones to have builds separated
Updating cloud builds to create one build per build configuration
Updating metadata for cloud builds to include BuildGroup
Updating name of cloud builds
Removing explicit HordeID from cloud builds

#rnx

diff --git a/Engine/Build/Graph/Tasks/BuildAndTestProject.xml b/Engine/Build/Graph/Tasks/BuildAndTestProject.xml
--- a/Engine/Build/Graph/Tasks/BuildAndTestProject.xml
+++ b/Engine/Build/Graph/Tasks/BuildAndTestProject.xml
@@ -469,7 +469,15 @@
 					<CreateArtifact Name="$(ProjectName)-Staged-Cooked-Editor-$(EditorPlatform)" Type="staged-build" BaseDir="$(TargetDirectory)" Keys="Project=$(ProjectName)" Metadata="Stream=$(Branch);Platform=$(EditorPlatform);Configuration=Development;BuildName=$(BuildNamePath);dashboard-category=Development Builds" Files="#Staged CookedEditor Files"/>
 				</Do>
 				<Do If="$(PublishCloudArtifacts)">
-					<CreateCloudArtifact Name="$(BuildNamePath)-Staged-Cooked-Editor" HordeArtifactName="Cloud-$(ProjectName)-Staged-Cooked-Editor-$(EditorPlatform)" Type="staged-build" BaseDir="$(TargetDirectory)" Metadata="Configuration=Development;dashboard-category=Staged Builds - Cloud" Project="$(TargetProject)" Platform="$(EditorPlatform)" Files="#Staged CookedEditor Files"/>
+					<CreateCloudArtifact 
+						Name="$(ProjectName)-$(TargetPlatform)-Staged-Cooked-Editor" 
+						BuildGroup="$(BuildVersion)" 
+						Type="staged-build" 
+						BaseDir="$(TargetDirectory)" 
+						Metadata="Configuration=Development;dashboard-category=Staged Builds - Cloud" 
+						Project="$(TargetProject)" 
+						Platform="$(EditorPlatform)" 
+						Files="#Staged CookedEditor Files"/>
 				</Do>
 				<Delete Files="$(TargetDirectory)/artifactmetadata.json"/>
 			</Node>
@@ -497,7 +505,12 @@
 			<Do If="$(RetrieveUsingCloudArtifacts)">
 				<Property Name="CloudArtifactPlatformRetrieve" Value="$(PreNodeName)Retrieve CookedEditor Build Artifact $(EditorPlatform) (Cloud)"/>
 				<Node Name="$(CloudArtifactPlatformRetrieve)" Requires="$(TestPrerequisites);$(PlatformPublishNodeName)">
-					<RetrieveCloudArtifact Name="$(BuildNamePath)-Staged-Cooked-Editor" Type="staged-build" OutputDir="$(ProjectOutputDirectory)/$(EditorPlatformPath)CookedEditor" Project="$(TargetProject)" Platform="$(EditorPlatform)" />
+					<RetrieveCloudArtifact 
+						Name="$(ProjectName)-$(TargetPlatform)-Staged-Cooked-Editor" 
+						Type="staged-build" 
+						OutputDir="$(ProjectOutputDirectory)/$(EditorPlatformPath)CookedEditor" 
+						Project="$(TargetProject)" 
+						Platform="$(EditorPlatform)" />
 				</Node>
 				<Property Name="TestPrerequisites" Value="$(TestPrerequisites);$(CloudArtifactPlatformRetrieve)" />
 				<Property Name="CookedEditorBuildPath" Value="$(ProjectOutputDirectory)/$(EditorPlatformPath)CookedEditor" />
@@ -844,7 +857,15 @@
 					<Expand Name="WriteArtifactMetadataFile" Type="staged-build" Platform="$(StagedPlatformFolder)" Configuration="$(PackageConfigurations)" ArtifactDir="$(TargetDirectory)" />
 					<Tag Files="$(TargetDirectory)/..." With="#$(PlatformStageNodeName) Files"/>
 
-					<CreateCloudArtifact Name="$(BuildNamePath)-Staged" HordeArtifactName="Cloud-$(ProjectName)-Staged-$(StagedPlatformFolder)" Type="staged-build" BaseDir="$(TargetDirectory)" Metadata="cookPlatform=$(StagedPlatformFolder);dashboard-category=Staged Builds - Cloud" Project="$(TargetProject)" Platform="$(StagedPlatformFolder)" Files="#$(PlatformStageNodeName) Files"/>
+					<CreateCloudArtifact 
+						Name="$(ProjectName)-$(TargetPlatform)-Staged" 
+						BuildGroup="$(BuildVersion)" 
+						Type="staged-build" 
+						BaseDir="$(TargetDirectory)" 
+						Metadata="cookPlatform=$(StagedPlatformFolder);dashboard-category=Staged Builds - Cloud" 
+						Project="$(TargetProject)" 
+						Platform="$(StagedPlatformFolder)" 
+						Files="#$(PlatformStageNodeName) Files"/>
 					<Delete Files="$(TargetDirectory)/artifactmetadata.json"/>
 				</Node>
 			</Do>
@@ -923,7 +944,15 @@
 						<Expand Name="WriteArtifactMetadataFile" Type="packaged-build" Platform="$(StagedPlatformFolder)" Configuration="$(PackageConfiguration)" ArtifactDir="$(TargetDirectory)" />
 						<Tag Files="$(TargetDirectory)/..." With="#$(PlatformPackageNodeName) $(PackageConfiguration) Files"/>
 
-						<CreateCloudArtifact Name="$(BuildNamePath)-Packaged-$(PackageConfiguration)" HordeArtifactName="Cloud-$(ProjectName)-Packaged-$(StagedPlatformFolder)-$(PackageConfiguration)" Type="packaged-build" BaseDir="$(TargetDirectory)" Metadata="cookPlatform=$(StagedPlatformFolder);Configuration=$(PackageConfiguration);dashboard-category=Packaged Builds - Cloud" Project="$(TargetProject)" Platform="$(StagedPlatformFolder)" Files="#$(PlatformPackageNodeName) $(PackageConfiguration) Files"/>
+						<CreateCloudArtifact 
+						Name="$(ProjectName)-$(TargetPlatform)-Packaged-$(PackageConfiguration)" 
+						BuildGroup="$(BuildVersion)"
+						Type="packaged-build" 
+						BaseDir="$(TargetDirectory)" 
+						Metadata="cookPlatform=$(StagedPlatformFolder);Configuration=$(PackageConfiguration);dashboard-category=Packaged Builds - Cloud" 
+						Project="$(TargetProject)" 
+						Platform="$(StagedPlatformFolder)" 
+						Files="#$(PlatformPackageNodeName) $(PackageConfiguration) Files"/>
 						<Delete Files="$(TargetDirectory)/artifactmetadata.json"/>
 					</ForEach>
 				</Node>
@@ -1034,10 +1063,20 @@
 							<Node Name="$(CloudArtifactPlatformRetrieveWithConfig)" Requires="$(TestPrerequisites)">
 								<Switch>
 									<Case If="$(UsePackaged)">
-										<RetrieveCloudArtifact Name="$(BuildNamePath)-Packaged-$(TestConfiguration)" Type="packaged-build" OutputDir="$(ProjectOutputDirectory)/Packages/$(StagedPlatformFolder)/$(TestConfiguration)" Project="$(TargetProject)" Platform="$(StagedPlatformFolder)" />
+										<RetrieveCloudArtifact 
+										Name="$(ProjectName)-$(TargetPlatform)-Packaged-$(TestConfiguration)" 
+										Type="packaged-build" 
+										OutputDir="$(ProjectOutputDirectory)/Packages/$(StagedPlatformFolder)/$(TestConfiguration)" 
+										Project="$(TargetProject)" 
+										Platform="$(StagedPlatformFolder)" />
 									</Case>
 									<Default>
-										<RetrieveCloudArtifact Name="$(BuildNamePath)-Staged" Type="staged-build" OutputDir="$(ProjectOutputDirectory)/$(StagedPlatformFolder)" Project="$(TargetProject)" Platform="$(StagedPlatformFolder)" />
+										<RetrieveCloudArtifact 
+										Name="$(ProjectName)-$(TargetPlatform)-Staged" 
+										Type="staged-build" 
+										OutputDir="$(ProjectOutputDirectory)/$(StagedPlatformFolder)" 
+										Project="$(TargetProject)" 
+										Platform="$(StagedPlatformFolder)" />
 									</Default>
 								</Switch>
 							</Node>
diff --git a/Engine/Source/Programs/AutomationTool/BuildGraph/BgScriptReader.cs b/Engine/Source/Programs/AutomationTool/BuildGraph/BgScriptReader.cs
--- a/Engine/Source/Programs/AutomationTool/BuildGraph/BgScriptReader.cs
+++ b/Engine/Source/Programs/AutomationTool/BuildGraph/BgScriptReader.cs
@@ -1815,12 +1815,12 @@ namespace AutomationTool
 		void AddCloudArtifactFromTask(BgScriptElement element, Dictionary<string, string> arguments)
 		{
 			arguments.TryGetValue(nameof(CreateCloudArtifactTaskParameters.Name), out string? name);
-			arguments.TryGetValue(nameof(CreateCloudArtifactTaskParameters.HordeArtifactName), out string? hordeArtifactName);
+			arguments.TryGetValue(nameof(CreateCloudArtifactTaskParameters.HordeArtifactID), out string? hordeArtifactID);
 			arguments.TryGetValue(nameof(CreateCloudArtifactTaskParameters.Type), out string? type);
 			arguments.TryGetValue(nameof(CreateCloudArtifactTaskParameters.Description), out string? description);
 			arguments.TryGetValue(nameof(CreateCloudArtifactTaskParameters.Metadata), out string? metadata);
 
-			string? artifactName = String.IsNullOrEmpty(hordeArtifactName) ? name : hordeArtifactName;
+			string? artifactName = String.IsNullOrEmpty(hordeArtifactID) ? name : hordeArtifactID;
 			if (String.IsNullOrEmpty(artifactName))
 			{
 				LogError(element, "Missing artifact name");
diff --git a/Engine/Source/Programs/AutomationTool/BuildGraph/Tasks/CreateCloudArtifactTask.cs b/Engine/Source/Programs/AutomationTool/BuildGraph/Tasks/CreateCloudArtifactTask.cs
--- a/Engine/Source/Programs/AutomationTool/BuildGraph/Tasks/CreateCloudArtifactTask.cs
+++ b/Engine/Source/Programs/AutomationTool/BuildGraph/Tasks/CreateCloudArtifactTask.cs
@@ -35,10 +35,16 @@ namespace AutomationTool.Tasks
 		public string Name { get; set; } = null!;
 
 		/// <summary>
-		/// Optional override of the name of the artifact as presented in Horde, defaults to name if not set
+		/// An optional override of the id of the artifact as presented in Horde, defaults to name if not set
 		/// </summary>
 		[TaskParameter(Optional = true)]
-		public string HordeArtifactName { get; set; } = null!;
+		public string HordeArtifactID { get; set; } = null!;
+
+		/// <summary>
+		/// An identifier to set for a group of builds that are compatible
+		/// </summary>
+		[TaskParameter]
+		public string BuildGroup { get; set; } = null!;
 
 		/// <summary>
 		/// The artifact type. Determines the permissions and expiration policy for the artifact.
@@ -194,15 +200,15 @@ namespace AutomationTool.Tasks
 		/// <param name="tagNameToFileSet">Mapping from tag names to the set of files they include.</param>
 		public override async Task ExecuteAsync(JobContext job, HashSet<FileReference> buildProducts, Dictionary<string, HashSet<FileReference>> tagNameToFileSet)
 		{
-			if (String.IsNullOrEmpty(_parameters.Name) && String.IsNullOrEmpty(_parameters.HordeArtifactName))
+			if (String.IsNullOrEmpty(_parameters.Name) && String.IsNullOrEmpty(_parameters.HordeArtifactID))
 			{
-				throw new AutomationException("No name set in either Name or HordeArtifactName");
+				throw new AutomationException("No name set in either Name or HordeArtifactID");
 			}
 
-			string name = String.IsNullOrEmpty(_parameters.Name) ? _parameters.HordeArtifactName : _parameters.Name;
+			string name = String.IsNullOrEmpty(_parameters.Name) ? _parameters.HordeArtifactID : _parameters.Name;
 			ArtifactType type = new ArtifactType(_parameters.Type);
-			string sanitizedHordeArtifactParam = SanitizeBucketValue(String.IsNullOrEmpty(_parameters.HordeArtifactName) ? name : _parameters.HordeArtifactName);
-			ArtifactName hordeArtifactName = new ArtifactName(sanitizedHordeArtifactParam);
+			string sanitizedHordeArtifactParam = SanitizeBucketValue(String.IsNullOrEmpty(_parameters.HordeArtifactID) ? name : _parameters.HordeArtifactID);
+			ArtifactName hordeArtifactID = new ArtifactName(sanitizedHordeArtifactParam);
 			List<string> metadataParam = (_parameters.Metadata ?? String.Empty).Split(';', StringSplitOptions.RemoveEmptyEntries | StringSplitOptions.TrimEntries).ToList();
 
 			Dictionary<string, object> metadata = new Dictionary<string, object>();
@@ -279,6 +285,11 @@ namespace AutomationTool.Tasks
 				metadata.Add("buildurl", $"{hordeUrl}job/{jobId}?step={stepId}");
 			}
 
+			if (!String.IsNullOrEmpty(_parameters.BuildGroup))
+			{
+				metadata.Add("buildgroup", _parameters.BuildGroup);
+			}
+
 			string httpVersion = _parameters.HttpVersion;
 			if (String.IsNullOrEmpty(httpVersion) || httpVersion == "None")
 			{
@@ -515,7 +526,7 @@ namespace AutomationTool.Tasks
 
 			try
 			{
-				await RegisterHordeArtifactAsync(buildUri, hordeArtifactName, type, streamIdTyped, commitId, metadata);
+				await RegisterHordeArtifactAsync(buildUri, hordeArtifactID, type, streamIdTyped, commitId, metadata);
 			}
 			catch (Exception ex)
 			{
