Updating CreateCloudArtifact calls to include BuildGroup

#rnx

diff --git a/Engine/Build/Graph/Examples/BuildEditorAndTools.xml b/Engine/Build/Graph/Examples/BuildEditorAndTools.xml
--- a/Engine/Build/Graph/Examples/BuildEditorAndTools.xml
+++ b/Engine/Build/Graph/Examples/BuildEditorAndTools.xml
@@ -198,7 +198,15 @@
 			<Error Message="UProjectPath is a required option when submitting to Cloud Storage" If="'$(UProjectPath)' == ''" />
 			<Property Name="UgsProjectKey" Value="ugs-project=$(Branch)/$(UProjectPath)" />
 
-			<CreateCloudArtifact Name="editor" Type="ugs-pcb" BaseDir="$(ArchiveStagingRelativeDir)" Metadata="$(UgsProjectKey);ArchiveType=Editor" Project="$(UProjectPath)" Platform="Windows" Files="#PublishBinaries"/>
+			<CreateCloudArtifact 
+				Name="editor" 
+				BuildGroup="$(EscapedBranch)-$(Change)"
+				Type="ugs-pcb" 
+				BaseDir="$(ArchiveStagingRelativeDir)" 
+				Metadata="$(UgsProjectKey);ArchiveType=Editor" 
+				Project="$(UProjectPath)" 
+				Platform="Windows" 
+				Files="#PublishBinaries"/>
 		</Node>
 
 		<Node Name="$(HordeSubmitNodeName)" Requires="Stage for UGS">
