IncrementalCompare : print cook diagnostics to the logs in the priming and null cook (the 2nd and 3rd cook).

#rnx

diff --git a/Engine/Build/Graph/Tests/IncrementalCompare.xml b/Engine/Build/Graph/Tests/IncrementalCompare.xml
--- a/Engine/Build/Graph/Tests/IncrementalCompare.xml
+++ b/Engine/Build/Graph/Tests/IncrementalCompare.xml
@@ -136,7 +136,7 @@
 			<Node Name="$(PrimingNullIncrementalCookNodeName)" Requires="$(StoreIncrementalCookNodeName)">
 				<!-- Execute the null incremental cook. -->
 				<Log Message="$(NewLine)========== IncrementalCompare.xml: Priming Null Incremental Cook =========="/>
-				<Property Name="NullCookArgs" Value="$(CommonCookArgs) -dpcvars=cook.displaymode=4"/>
+				<Property Name="NullCookArgs" Value="$(CommonCookArgs) -dpcvars=cook.displaymode=4 -CookIncrementallyModifiedDiagnostics -WriteCookIncrementallyModifiedDiagnosticsToLogs"/>
 				<Commandlet Name="Cook" Project="$(RelativeProjectFile)" Arguments="$(NullCookArgs)"/>
 			</Node>
 			<Node Name="$(PublishingPrimingNullIncrementalCookNodeName)" Requires="$(PrimingNullIncrementalCookNodeName)">
@@ -148,7 +148,7 @@
 			<Node Name="$(NullIncrementalCookNodeName)" Requires="$(PublishingPrimingNullIncrementalCookNodeName)">
 				<!-- Execute the null incremental cook. -->
 				<Log Message="$(NewLine)========== IncrementalCompare.xml: Null Incremental Cook =========="/>
-				<Property Name="NullCookArgs" Value="$(CommonCookArgs) -dpcvars=cook.displaymode=4"/>
+				<Property Name="NullCookArgs" Value="$(CommonCookArgs) -dpcvars=cook.displaymode=4 -CookIncrementallyModifiedDiagnostics -WriteCookIncrementallyModifiedDiagnosticsToLogs"/>
 				<Commandlet Name="Cook" Project="$(RelativeProjectFile)" Arguments="$(NullCookArgs)"/>
 			</Node>
 			<Node Name="$(PublishingNullIncrementalCookNodeName)" Requires="$(NullIncrementalCookNodeName)">
