changing targetplatform to editor platform for cloud uploads

#rnx

diff --git a/Engine/Build/Graph/Tasks/BuildAndTestProject.xml b/Engine/Build/Graph/Tasks/BuildAndTestProject.xml
--- a/Engine/Build/Graph/Tasks/BuildAndTestProject.xml
+++ b/Engine/Build/Graph/Tasks/BuildAndTestProject.xml
@@ -470,7 +470,7 @@
 				</Do>
 				<Do If="$(PublishCloudArtifacts)">
 					<CreateCloudArtifact 
-						Name="$(ProjectName)-$(TargetPlatform)-Staged-Cooked-Editor" 
+						Name="$(ProjectName)-$(EditorPlatform)-Staged-Cooked-Editor" 
 						BuildGroup="$(BuildVersion)" 
 						Type="staged-build" 
 						BaseDir="$(TargetDirectory)" 
@@ -506,7 +506,7 @@
 				<Property Name="CloudArtifactPlatformRetrieve" Value="$(PreNodeName)Retrieve CookedEditor Build Artifact $(EditorPlatform) (Cloud)"/>
 				<Node Name="$(CloudArtifactPlatformRetrieve)" Requires="$(TestPrerequisites);$(PlatformPublishNodeName)">
 					<RetrieveCloudArtifact 
-						Name="$(ProjectName)-$(TargetPlatform)-Staged-Cooked-Editor" 
+						Name="$(ProjectName)-$(EditorPlatform)-Staged-Cooked-Editor" 
 						Type="staged-build" 
 						OutputDir="$(ProjectOutputDirectory)/$(EditorPlatformPath)CookedEditor" 
 						Project="$(TargetProject)" 
