Allow post editor build steps and custom pre-cook requirements

diff --git a/Engine/Build/Graph/Tasks/BuildAndTestProject.xml b/Engine/Build/Graph/Tasks/BuildAndTestProject.xml
--- a/Engine/Build/Graph/Tasks/BuildAndTestProject.xml
+++ b/Engine/Build/Graph/Tasks/BuildAndTestProject.xml
@@ -152,9 +152,13 @@
 
 	<Option Name="TestsRequireEditor" DefaultValue="false" Description="Whether to force tests to require editor compilation"/>
 
-	<!-- Macros that allows custom pre-steps -->
+	<!-- Macros that allows custom post-editor build and pre-cook steps -->
 	<Macro Name="PreCookStep"/>
-	
+	<Macro Name="PostEditorPlatformStep"/>
+
+	<!-- Macro that allows user to specify additional cook requirements by populating $(PlatformCookAdditionalRequirements) property -->
+	<Macro Name="CookAdditionalRequirements"/>
+
 	<!-- Macros that allow custom post tests steps -->
 	<Option Name="PostEditorTestsExtendFile" DefaultValue="" Description="Path to the file containing all post editor tests macro extends"/>
 	<Macro Name="PostEditorTestsStep"/>
@@ -422,6 +426,8 @@
 				<Property Name="GeneratedEditorTestNodes" Value="$(GeneratedEditorTestNodes);$(PlatformEditorTestNodes)" />
 			</Do>
 		</Agent>
+
+		<Expand Name="PostEditorPlatformStep"/>
 	</ForEach>
 
 	<!-- Cooked editor Targets -->
@@ -718,12 +724,15 @@
 
 		</Agent>
 
+		<Property Name="PlatformCookAdditionalRequirements" Value=""/>
+		<Expand Name="CookAdditionalRequirements"/>
+
 		<Property Name="PlatformEditorNodeName" Value="" If="$(SkipEditorBuild)" />
 		<!-- Separate agent for cooking / staging, so the cook can happen in parallel with the client compile -->
 		<Agent Name="Package $(PlatformAgentName)" Type="$(PackageAgentType)">
 
 			<!-- Cook -->
-			<Node Name="$(PlatformCookNodeName)" Requires="$(PlatformEditorNodeName)" Produces="#$(PlatformCookNodeName) Complete;#$(PlatformCookNodeName) CookShaderSymbols" RunEarly="true">
+			<Node Name="$(PlatformCookNodeName)" Requires="$(PlatformEditorNodeName);$(PlatformCookAdditionalRequirements)" Produces="#$(PlatformCookNodeName) Complete;#$(PlatformCookNodeName) CookShaderSymbols" RunEarly="true">
 				<Do If="!$(SkipCook)">
 					<Expand Name="PreCookStep"/>
 					<Property Name="CookArgs" Value="" />
@@ -738,7 +747,7 @@
 				</Do>
 			</Node>
 			
-			<Node Name="$(PlatformCookNodeName) Server" Requires="$(PlatformEditorNodeName)" Produces="#$(PlatformCookNodeName) Server Complete" RunEarly="true">
+			<Node Name="$(PlatformCookNodeName) Server" Requires="$(PlatformEditorNodeName);$(PlatformCookAdditionalRequirements)" Produces="#$(PlatformCookNodeName) Server Complete" RunEarly="true">
 				<Do If="!$(SkipCook)">
 					<Expand Name="PreCookStep"/>
 					<Property Name="CookArgs" Value="" />
