Use Include property IgnoreDuplicateIncludes to avoid duplicating include of the TestConfigs file (remove workaround)

#rnx

diff --git a/Engine/Build/Graph/Tasks/BuildAndTestProject.xml b/Engine/Build/Graph/Tasks/BuildAndTestProject.xml
--- a/Engine/Build/Graph/Tasks/BuildAndTestProject.xml
+++ b/Engine/Build/Graph/Tasks/BuildAndTestProject.xml
@@ -190,7 +190,7 @@
 	</Do>
 
 	<!-- Test configs -->
-	<Include Script="Inc/TestConfigs.xml"/>
+	<Include Script="Inc/TestConfigs.xml" IgnoreDuplicateIncludes="true"/>
 	<Do If="'$(TestConfigList)' == ''">
 		<Expand Name="AddTestConfig" TestConfigName=""/>
 	</Do>
diff --git a/Engine/Build/Graph/Tasks/Inc/TestConfigs.xml b/Engine/Build/Graph/Tasks/Inc/TestConfigs.xml
--- a/Engine/Build/Graph/Tasks/Inc/TestConfigs.xml
+++ b/Engine/Build/Graph/Tasks/Inc/TestConfigs.xml
@@ -15,10 +15,8 @@
 
 	<Option Name="TestConfigList" DefaultValue="" Description="Input TestConfigs; expand macro AddTestConfig after including Engine/Build/Graph/Tasks/Inc/TestConfigs.xml to populate"/>
 
-	<Do If="'$(TestConfigList)' == ''">
-		<Macro Name="AddTestConfig" Arguments="TestConfigName" OptionalArguments="TestPlatform;AdditionalTestArguments;ArtifactPostfix;AdditionalAppArguments;AgentType">
-			<Property Name="TestConfigList" Value="$(TestConfigList)^$(TestConfigName)|$(TestPlatform)|$(ArtifactPostfix)|$(AdditionalTestArguments)|$(AdditionalAppArguments)|$(AgentType)"/>
-		</Macro>
-	</Do>
+	<Macro Name="AddTestConfig" Arguments="TestConfigName" OptionalArguments="TestPlatform;AdditionalTestArguments;ArtifactPostfix;AdditionalAppArguments;AgentType">
+		<Property Name="TestConfigList" Value="$(TestConfigList)^$(TestConfigName)|$(TestPlatform)|$(ArtifactPostfix)|$(AdditionalTestArguments)|$(AdditionalAppArguments)|$(AgentType)"/>
+	</Macro>
 
 </BuildGraph>
\ No newline at end of file
