Fix for APT RunLocalTests script failing for foreign code projects.

These scripts use BuildAndTestProject.xml. Two of the steps fail because the project could not be found.

diff --git a/Engine/Build/Graph/Tasks/BuildAndTestProject.xml b/Engine/Build/Graph/Tasks/BuildAndTestProject.xml
--- a/Engine/Build/Graph/Tasks/BuildAndTestProject.xml
+++ b/Engine/Build/Graph/Tasks/BuildAndTestProject.xml
@@ -345,7 +345,7 @@
 						<TagReceipt Files="#$(CompileEditor)" BuildProducts="true" With="#$(CompileEditor)"/>
 					</Case>
 					<Default>
-						<Compile Target="$(EditorTargetName)" Platform="$(EditorPlatform)" Configuration="Development" Arguments="$(GenericCompileArguments) $(ExtraEditorCompileArguments)"/>
+						<Compile Target="$(EditorTargetName)" Project="$(TargetProject)" Platform="$(EditorPlatform)" Configuration="Development" Arguments="$(GenericCompileArguments) $(ExtraEditorCompileArguments)"/>
 					</Default>
 				</Switch>
 			</Node>
@@ -699,7 +699,7 @@
 					</Case>
 					<Default>
 						<ForEach Name="TargetConfiguration" Values="$(TargetConfigurations)" Separator="+">
-							<Compile Target="$(TargetName)" Platform="$(TargetPlatform)" Configuration="$(TargetConfiguration)" Arguments="-BuildVersion=&quot;$(BuildVersion)&quot; $(GenericCompileArguments) $(TargetCompileArguments)"/>
+							<Compile Target="$(TargetName)" Project="$(TargetProject)" Platform="$(TargetPlatform)" Configuration="$(TargetConfiguration)" Arguments="-BuildVersion=&quot;$(BuildVersion)&quot; $(GenericCompileArguments) $(TargetCompileArguments)"/>
 						</ForEach>
 					</Default>
 				</Switch>
