Remove test data v1 output

#rnx

diff --git a/Engine/Build/Graph/Tasks/Inc/GauntletSettings.xml b/Engine/Build/Graph/Tasks/Inc/GauntletSettings.xml
--- a/Engine/Build/Graph/Tasks/Inc/GauntletSettings.xml
+++ b/Engine/Build/Graph/Tasks/Inc/GauntletSettings.xml
@@ -25,13 +25,15 @@
 	<Option Name="ExtraGauntletCommonArgs" DefaultValue="" Description="Additional Common arguments passed to Gauntlet command line"/>
 	<Option Name="ExtraGauntletTargetTestArgs" DefaultValue="" Description="Additional Target test arguments passed to Gauntlet command line"/>
 	<Option Name="ExtraGauntletEditorTestArgs" DefaultValue="" Description="Additional Editor test arguments passed to Gauntlet command line"/>
+	<Option Name="UseTestDataV2" DefaultValue="$(Horde)" Description="Publish report as TestData v2 for Horde"/>
 
 	<Option Name="AdbKeysPath" DefaultValue="" Description="Path to adbkeys location (for Android)"/>
 
 	<Option Name="DisableTestResultsForHorde" DefaultValue="false" Restrict="true|false" Description="Force disabling test results for Horde production"/>
 
 	<Property Name="GauntletCommonArgs" Value="-maxduration=$(GauntletTimeOut) -unattended -verbose -ResumeOnCriticalFailure $(ExtraGauntletCommonArgs)" />
-	<Property Name="GauntletCommonArgs" Value="$(GauntletCommonArgs) -WriteTestResultsForHorde -UseTestDataV2 -AddTestDataV1" If="$(Horde) and !$(DisableTestResultsForHorde)"/>
+	<Property Name="GauntletCommonArgs" Value="$(GauntletCommonArgs) -WriteTestResultsForHorde -UseTestDataV2" If="$(Horde) and !$(DisableTestResultsForHorde)"/>
+	<Property Name="GauntletCommonArgs" Value="$(GauntletCommonArgs) -UseTestDataV2" If="$(Horde) and !$(DisableTestResultsForHorde) and $(UseTestDataV2)"/>
 	<Property Name="GauntletCommonArgs" Value="$(GauntletCommonArgs) -branch=&quot;$(Branch)&quot; -changelist=&quot;$(Change)&quot;" />
 
 	<Property Name="GauntletDeviceArgs" Value="" />
