Adding preflight metadata to cloud artifacts

#rnx

diff --git a/Engine/Build/Graph/Tasks/BuildAndTestProject.xml b/Engine/Build/Graph/Tasks/BuildAndTestProject.xml
--- a/Engine/Build/Graph/Tasks/BuildAndTestProject.xml
+++ b/Engine/Build/Graph/Tasks/BuildAndTestProject.xml
@@ -482,7 +482,7 @@
 						BuildGroup="$(BuildVersion)" 
 						Type="staged-build" 
 						BaseDir="$(TargetDirectory)" 
-						Metadata="Configuration=Development;dashboard-category=Staged Builds - Cloud" 
+						Metadata="Configuration=Development;dashboard-category=Staged Builds - Cloud;$(SharedCloudBuildMetadata)" 
 						Project="$(TargetProject)" 
 						Platform="$(EditorPlatform)" 
 						Files="#Staged CookedEditor Files"/>
@@ -873,7 +873,7 @@
 						BuildGroup="$(BuildVersion)" 
 						Type="staged-build" 
 						BaseDir="$(TargetDirectory)" 
-						Metadata="cookPlatform=$(StagedPlatformFolder);dashboard-category=Staged Builds - Cloud" 
+						Metadata="cookPlatform=$(StagedPlatformFolder);dashboard-category=Staged Builds - Cloud;$(SharedCloudBuildMetadata)" 
 						Project="$(TargetProject)" 
 						Platform="$(StagedPlatformFolder)" 
 						Files="#$(PlatformStageNodeName) Files"/>
@@ -960,7 +960,7 @@
 						BuildGroup="$(BuildVersion)"
 						Type="packaged-build" 
 						BaseDir="$(TargetDirectory)" 
-						Metadata="cookPlatform=$(StagedPlatformFolder);Configuration=$(PackageConfiguration);dashboard-category=Packaged Builds - Cloud" 
+						Metadata="cookPlatform=$(StagedPlatformFolder);Configuration=$(PackageConfiguration);dashboard-category=Packaged Builds - Cloud;$(SharedCloudBuildMetadata)" 
 						Project="$(TargetProject)" 
 						Platform="$(StagedPlatformFolder)" 
 						Files="#$(PlatformPackageNodeName) $(PackageConfiguration) Files"/>
diff --git a/Engine/Build/Graph/Tasks/Inc/CommonProperties.xml b/Engine/Build/Graph/Tasks/Inc/CommonProperties.xml
--- a/Engine/Build/Graph/Tasks/Inc/CommonProperties.xml
+++ b/Engine/Build/Graph/Tasks/Inc/CommonProperties.xml
@@ -21,6 +21,7 @@
 	<Property Name="BuildNamePath" Value="$(ProjectName)-$(CL-String)"/>
 	<Option Name="BuildVersion" DefaultValue="$(ProjectName)-CL-$(Change)" Description="Value to set as BuildVersion during compilation"/>
 
+	<Property Name="SharedCloudBuildMetadata" Value="ispreflight=$(ISPreflight);preflightchange=$(PreflightChange);isbuildmachine=$(IsBuildMachine)"/>
 	<Include Script="SharedProperties.xml" />
 
 </BuildGraph>
