Changing the name metadata field and adding buildGroup metadata field when storing snapshots to adhere to metadata standards aligned with staged and packaged builds.

diff --git a/Engine/Build/Graph/Tasks/BuildAndTestProject.xml b/Engine/Build/Graph/Tasks/BuildAndTestProject.xml
--- a/Engine/Build/Graph/Tasks/BuildAndTestProject.xml
+++ b/Engine/Build/Graph/Tasks/BuildAndTestProject.xml
@@ -817,7 +817,7 @@
 						<!-- Publish to builds cloud storage endpoint -->
 						<Property Name="SnapshotBuildBaselineBranch" Value="$(Branch)"/>
 						<Property Name="SnapshotBuildBaselineBranch" Value="$(UE-CloudPublishBaselineBranch)" If="'$(UE-CloudPublishBaselineBranch)' != ''"/>
-						<Property Name="SnapshotBuildMetaData" Value="name=$(SnapshotIdentifier);depot=$(Depot);branch=$(Branch);baselineBranch=$(SnapshotBuildBaselineBranch);project=$(ProjectName);runtime=$(TargetType);platform=$(TargetPlatform);cookPlatform=$(CookPlatform);changelist=$(Change);commit=$(Change);codeChangelist=$(CodeChange);codeCommit=$(CodeChange);isPreflight=$(IsPreflight);isBuildMachine=$(IsBuildMachine);isHordeEnv=$(IsHordeEnv);buildgroup=$(BuildVersion)"/>
+						<Property Name="SnapshotBuildMetaData" Value="name=$(CookPlatform);buildGroup=$(BuildVersion);depot=$(Depot);branch=$(Branch);baselineBranch=$(SnapshotBuildBaselineBranch);project=$(ProjectName);runtime=$(TargetType);platform=$(TargetPlatform);cookPlatform=$(CookPlatform);changelist=$(Change);commit=$(Change);codeChangelist=$(CodeChange);codeCommit=$(CodeChange);isPreflight=$(IsPreflight);isBuildMachine=$(IsBuildMachine);isHordeEnv=$(IsHordeEnv);buildgroup=$(BuildVersion)"/>
 						<Property Name="SnapshotBuildMetaData" Value="$(SnapshotBuildMetaData);cookFlavor=$(BCRCookFlavor)" If="'$(BCRCookFlavor)' != ''" />
 						<Property Name="SnapshotBuildMetaData" Value="$(SnapshotBuildMetaData);preflightChangelist=$(PreflightChange)" If="'$(PreflightChange)' != ''" />
 						<Property Name="SnapshotBuildMetaData" Value="$(SnapshotBuildMetaData);hordeTemplateId=$(UE_HORDE_TEMPLATEID)" If="'$(UE_HORDE_TEMPLATEID)' != ''" />
