Disabling horde artifact uploads when we are using cloud artifact uploads

#rnx

diff --git a/Engine/Build/Graph/Tasks/BuildAndTestProject.xml b/Engine/Build/Graph/Tasks/BuildAndTestProject.xml
--- a/Engine/Build/Graph/Tasks/BuildAndTestProject.xml
+++ b/Engine/Build/Graph/Tasks/BuildAndTestProject.xml
@@ -244,7 +244,9 @@
 	<Option Name="RetrieveUsingCloudArtifacts" DefaultValue="false" Restrict="true|false" Description="Enable to retrieve cloud artifacts" />
 	<Option Name="PublishToNetworkOutputDirectory" DefaultValue="$(IsBuildMachine)" Restrict="true|false" Description="Enable to publish artifacts to network output directory" />
 	<Option Name="PublishCloudArtifacts" DefaultValue="$(RetrieveUsingCloudArtifacts)" Description="Enable to publish cloud artifacts" />
-	<Option Name="PublishHordeArtifacts" DefaultValue="$(IsEpicBuildMachine)" Description="Enable to publish horde artifacts"/>
+	<Property Name="ShouldPublishHordeArtifact" Value="false" />
+	<Property Name="ShouldPublishHordeArtifact" Value="true" If="$(IsEpicBuildMachine) and !$(PublishCloudArtifacts)"/>
+	<Option Name="PublishHordeArtifacts" DefaultValue="$(ShouldPublishHordeArtifact)" Description="Enable to publish horde artifacts"/>
 
 	<!-- If running with GPU Reshape, override the local build configuration
 	     Enables shader symbols and anything else that aids troubleshooting -->
