Avoid executing SymStore node on linux as it is not supporting symbol sever

#rnx

diff --git a/Engine/Build/Graph/Tasks/BuildAndTestProject.xml b/Engine/Build/Graph/Tasks/BuildAndTestProject.xml
--- a/Engine/Build/Graph/Tasks/BuildAndTestProject.xml
+++ b/Engine/Build/Graph/Tasks/BuildAndTestProject.xml
@@ -628,7 +628,7 @@
 		<Property Name="PlatformAgentName" Value="$(ProjectName) $(TargetPlatform)" />
 		<Property Name="PlatformCompileNodeName" Value="$(PreNodeNameWithPlatform)Compile $(TargetPlatform)" />
 		<Property Name="UploadSymbolsNodeName" Value="" />
-		<Property Name="UploadSymbolsNodeName" Value="$(PreNodeNameWithPlatform)Upload Symbols $(TargetPlatform)" If="'$(SymbolServerPath)' != '' And !$(IsMobilePlatform)" />
+		<Property Name="UploadSymbolsNodeName" Value="$(PreNodeNameWithPlatform)Upload Symbols $(TargetPlatform)" If="'$(SymbolServerPath)' != '' And !$(IsMobilePlatform) And !$(IsTargetPlatformLinux)" />
 		<Property Name="PlatformCookNodeName" Value="$(PreNodeNameWithPlatform)Cook $(TargetPlatform)" />
 		<Property Name="PlatformStageNodeName" Value="$(PreNodeNameWithPlatform)Stage $(TargetPlatform)" />
 		<Property Name="PlatformPackageNodeName" Value="$(PreNodeNameWithPlatform)Package $(TargetPlatform)" />
