Add alternate platform cook during incremental compare to help evaluate whether cooking for another platform affects null incremental cook.

diff --git a/Engine/Build/Graph/Tests/IncrementalCompare.xml b/Engine/Build/Graph/Tests/IncrementalCompare.xml
--- a/Engine/Build/Graph/Tests/IncrementalCompare.xml
+++ b/Engine/Build/Graph/Tests/IncrementalCompare.xml
@@ -38,15 +38,17 @@
 	<Option Name="MPCookCoreSubscription" DefaultValue="" Description="Over or Under subscription factor for CPU cores to use during cooking, default is ini-driven."/>
 	<Option Name="WithNullCook" Restrict="true|false" DefaultValue="false" Description="Whether to do a null cook or not"/>
 	<Option Name="NullCookCustomConfig" DefaultValue="" Description="What custom configuration to apply if doing a null cook (see -CustomConfig commandline argument)"/>
+	<Option Name="NullCookInterleaveAlternatePlatform" DefaultValue="" Description="Optional platform cook to do betwen priming null cook and null cook (useful for testing independence of platforms when cooking)"/>
 	<Option Name="PublishCloudSnapshot" DefaultValue="false" Restrict="true|false" Description="Publish build to cloud host" />
 
 	<Property Name="FullPathProjectDir" Value="$(RootDir)/$(RelativeProjectDir)"/>
 	<Property Name="EditorPlatform" Value="Win64"/>
-	<Property Name="CommonCookArgs" Value="-targetplatform=$(TargetPlatform) -cookcultures=en -unversioned -zenstore -cookincremental"/>
-	<Property Name="CommonCookArgs" Value="$(CommonCookArgs) -CookProcessCount=$(CookProcessCount)" If="'$(CookProcessCount)' != ''"/>
-	<Property Name="CommonCookArgs" Value="$(CommonCookArgs) -MPCookCoreSubscription=$(MPCookCoreSubscription)" If="'$(MPCookCoreSubscription)' != ''"/>
-	<Property Name="CommonCookArgs" Value="$(CommonCookArgs) -CustomConfig=$(NullCookCustomConfig)" If="'$(WithNullCook) and $(NullCookCustomConfig)' != ''"/>
-	<Property Name="CommonCookArgs" Value="$(CommonCookArgs) $(ExtraCookArgs)" If="'$(ExtraCookArgs)' != ''"/>
+	<Property Name="CommonCookArgsNoTargetPlatform" Value="-cookcultures=en -unversioned -zenstore -cookincremental"/>
+	<Property Name="CommonCookArgsNoTargetPlatform" Value="$(CommonCookArgsNoTargetPlatform) -CookProcessCount=$(CookProcessCount)" If="'$(CookProcessCount)' != ''"/>
+	<Property Name="CommonCookArgsNoTargetPlatform" Value="$(CommonCookArgsNoTargetPlatform) -MPCookCoreSubscription=$(MPCookCoreSubscription)" If="'$(MPCookCoreSubscription)' != ''"/>
+	<Property Name="CommonCookArgsNoTargetPlatform" Value="$(CommonCookArgsNoTargetPlatform) -CustomConfig=$(NullCookCustomConfig)" If="'$(WithNullCook) and $(NullCookCustomConfig)' != ''"/>
+	<Property Name="CommonCookArgsNoTargetPlatform" Value="-targetplatform=$(TargetPlatform) $(CommonCookArgsNoTargetPlatform)"/>
+	<Property Name="CommonCookArgs" Value="$(CommonCookArgsNoTargetPlatform) $(ExtraCookArgs)" If="'$(ExtraCookArgs)' != ''"/>
 	<Property Name="PlatformCookedDir" Value="$(FullPathProjectDir)/Saved/Cooked/$(TargetPlatform)"/>
 	<Property Name="CompareDirIncremental" Value="$(FullPathProjectDir)/Intermediate/IncrementalCompare/Iter/$(TargetPlatform)"/>
 	<Property Name="CompareDirRecook" Value="$(FullPathProjectDir)/Intermediate/IncrementalCompare/Recook/$(TargetPlatform)"/>
@@ -57,6 +59,7 @@
 	<Property Name="StoreIncrementalCookNodeName" Value="Store Incremental Cook $(ProjectName) $(TargetPlatform)"/>
 	<Property Name="PrimingNullIncrementalCookNodeName" Value="Priming Null Incremental Cook $(ProjectName) $(TargetPlatform)"/>
 	<Property Name="PublishingPrimingNullIncrementalCookNodeName" Value="Publishing Priming Null Incremental Cook $(ProjectName) $(TargetPlatform)"/>
+	<Property Name="AlternateInterleavedPlatformCookNodeName" Value="Alternate Platform Cook $(ProjectName) $(NullCookInterleaveAlternatePlatform)" If="'$(NullCookInterleaveAlternatePlatform)' != ''"/>
 	<Property Name="NullIncrementalCookNodeName" Value="Null Incremental Cook $(ProjectName) $(TargetPlatform)"/>
 	<Property Name="PublishingNullIncrementalCookNodeName" Value="Publishing Null Incremental Cook $(ProjectName) $(TargetPlatform)"/>
 	<Property Name="CleanNullIncrementalCookNodeName" Value="Clean Null Incremental Cook $(ProjectName) $(TargetPlatform)"/>
@@ -145,7 +148,16 @@
 				<!-- Exporting snapshot for developer debugging needs -->
 				<Expand Name="Snapshot Publish" NameSuffix="AfterPrimingNullCook"/>
 			</Node>
-			<Node Name="$(NullIncrementalCookNodeName)" Requires="$(PublishingPrimingNullIncrementalCookNodeName)">
+			<Property Name="NullIncrementalCookPrereq" Value="$(PublishingPrimingNullIncrementalCookNodeName)"/>
+			<Do If="'$(NullCookInterleaveAlternatePlatform)' != ''">
+				<Property Name="NullIncrementalCookPrereq" Value="$(AlternateInterleavedPlatformCookNodeName)"/>
+				<Node Name="$(AlternateInterleavedPlatformCookNodeName)" Requires="$(PublishingPrimingNullIncrementalCookNodeName)">
+					<!-- Execute the alternate platform cook. -->
+					<Log Message="$(NewLine)========== IncrementalCompare.xml: Alternate Platform Cook =========="/>
+					<Commandlet Name="Cook" Project="$(RelativeProjectFile)" Arguments="-targetplatform=$(NullCookInterleaveAlternatePlatform) $(CommonCookArgsNoTargetPlatform) -forcerecook"/>
+				</Node>
+			</Do>
+			<Node Name="$(NullIncrementalCookNodeName)" Requires="$(NullIncrementalCookPrereq)">
 				<!-- Execute the null incremental cook. -->
 				<Log Message="$(NewLine)========== IncrementalCompare.xml: Null Incremental Cook =========="/>
 				<Property Name="NullCookArgs" Value="$(CommonCookArgs) -dpcvars=cook.displaymode=4 -CookIncrementallyModifiedDiagnostics -WriteCookIncrementallyModifiedDiagnosticsToLogs"/>
