adding buildversion and allowmultiplematches to retrievecloudartifact tasks to allow them to function properly when multiple preflights are made on the same base CL

#rnx

diff --git a/Engine/Build/Graph/Tasks/BuildAndTestProject.xml b/Engine/Build/Graph/Tasks/BuildAndTestProject.xml
--- a/Engine/Build/Graph/Tasks/BuildAndTestProject.xml
+++ b/Engine/Build/Graph/Tasks/BuildAndTestProject.xml
@@ -520,7 +520,10 @@
 						Type="staged-build" 
 						OutputDir="$(ProjectOutputDirectory)/$(EditorPlatformPath)CookedEditor" 
 						Project="$(TargetProject)" 
-						Platform="$(EditorPlatform)" />
+						Platform="$(EditorPlatform)" 		
+						BuildGroup="$(BuildVersion)" 
+						AllowMultipleMatches="true"
+						/>
 				</Node>
 				<Property Name="TestPrerequisites" Value="$(TestPrerequisites);$(CloudArtifactPlatformRetrieve)" />
 				<Property Name="CookedEditorBuildPath" Value="$(ProjectOutputDirectory)/$(EditorPlatformPath)CookedEditor" />
@@ -1081,7 +1084,10 @@
 										Type="packaged-build" 
 										OutputDir="$(ProjectOutputDirectory)/Packages/$(StagedPlatformFolder)/$(TestConfiguration)" 
 										Project="$(TargetProject)" 
-										Platform="$(StagedPlatformFolder)" />
+										Platform="$(StagedPlatformFolder)"
+										BuildGroup="$(BuildVersion)" 
+										AllowMultipleMatches="true"
+										/>
 									</Case>
 									<Default>
 										<RetrieveCloudArtifact 
@@ -1089,7 +1095,10 @@
 										Type="staged-build" 
 										OutputDir="$(ProjectOutputDirectory)/$(StagedPlatformFolder)" 
 										Project="$(TargetProject)" 
-										Platform="$(StagedPlatformFolder)" />
+										Platform="$(StagedPlatformFolder)" 
+										BuildGroup="$(BuildVersion)" 
+										AllowMultipleMatches="true"
+										/>
 									</Default>
 								</Switch>
 							</Node>
