Allow publishing snapshots from preflights.

diff --git a/Engine/Build/Graph/Tasks/BuildAndTestProject.xml b/Engine/Build/Graph/Tasks/BuildAndTestProject.xml
--- a/Engine/Build/Graph/Tasks/BuildAndTestProject.xml
+++ b/Engine/Build/Graph/Tasks/BuildAndTestProject.xml
@@ -801,9 +801,9 @@
 			<Property Name="ShouldSnapshotPublishCloud" Value="false"/>
 			<Property Name="ShouldSnapshotPublishCloud" Value="true" If="$(PublishCloudSnapshot) and '$(UE-CloudPublishHost)' != ''"/>
 			<Property Name="ShouldSnapshotPublishZen" Value="false"/>
-			<Property Name="ShouldSnapshotPublishZen" Value="true" If="$(IsBuildMachine) and !$(IsPreflight) and $(PublishZenSnapshot) and '$(UE-ZenPublishHost)' != ''"/>
+			<Property Name="ShouldSnapshotPublishZen" Value="true" If="$(IsBuildMachine) and $(PublishZenSnapshot) and '$(UE-ZenPublishHost)' != ''"/>
 			<Property Name="ShouldSnapshotPublish" Value="false"/>
-			<Property Name="ShouldSnapshotPublish" Value="true" If="$(IsBuildMachine) and !$(IsPreflight) and $(ShouldSnapshotPublishCloud) or $(ShouldSnapshotPublishZen)"/>
+			<Property Name="ShouldSnapshotPublish" Value="true" If="$(IsBuildMachine) and $(ShouldSnapshotPublishCloud) or $(ShouldSnapshotPublishZen)"/>
 			<Do If="$(PublishBuilds) and $(ShouldSnapshotPublish)">
 				<Property Name="PlatformSnapshotPublishNodeName" Value="$(PreNodeName)Snapshot Publish $(TargetPlatform)" />
 				<Node Name="$(PlatformSnapshotPublishNodeName)" Requires="$(PlatformHashNodeName)">
