UE-314026 - (Win64) Infrastructure Zen Streaming for APT

Updates the build graphs and C# files to add Win64 A/B Zen Streaming Import package test infrastructure to APT & BuildAndTestProject. Multi-platform support will come in a future CR.

diff --git a/Engine/Build/Graph/Tasks/BuildAndTestProject.xml b/Engine/Build/Graph/Tasks/BuildAndTestProject.xml
--- a/Engine/Build/Graph/Tasks/BuildAndTestProject.xml
+++ b/Engine/Build/Graph/Tasks/BuildAndTestProject.xml
@@ -12,6 +12,7 @@
 	<EnvVar Name="UE-CloudPublishNamespace"/>
 	<EnvVar Name="UE-CloudPublishBaselineBranch"/>
 	<EnvVar Name="UE-ZenPublishHost"/>
+	<EnvVar Name="uebp_BuildRoot_Escaped"/>
 
 	<Property Name="IsHordeEnv" Value="false"/>
 	<Property Name="IsHordeEnv" Value="true" If="'$(UE_HORDE_JOBID)' != ''"/>
@@ -66,7 +67,15 @@
 	<Option Name="IncludeEpicWebHelperInToolCompile" DefaultValue="false" Description="Include EpicWebHelper into Compile Tools prerequisites" />
 	<Option Name="WithDeviceReservation" DefaultValue="$(Horde)" Description="Whether to re-use devices and installed builds for a sequence of tests, where supported"/>
 	<Option Name="DoGPUReshape" DefaultValue="false" Restrict="true|false" Description="Whether to run GPU instrumentation test with GPU Reshape." />
-
+	<Option Name="PreflightPublishOverride" DefaultValue="false" Restrict="true|false" Description="Set to True to override Preflight Publishing when verifying gated behaviors." />
+
+	<!-- Zen Streaming Config -->
+	<Option Name="DoBuildZenStreaming" DefaultValue="true" Restrict="true|false" Description="A flag to enable ZenStreaming Building, Cooking, and Testing"/>
+	<Option Name="DoLocalZenTest" DefaultValue="true" Restrict="true|false" Description="Used to flag whether or not to use a local Zen Server for testing performance."/>
+	<Option Name="ZenStoreHost" DefaultValue="" Description="The hostname or ip of the target zenserver. Do not include http://" />
+	<Option Name="ZenStorePort" DefaultValue="8558" Description="The port to use on the target zen server." />
+	<Option Name="ZenStreamingWinTestPool" DefaultValue="" Description="Override the default Win64 Zen Streaming Test Agent Pool." />
+	
 	<!-- 
 		Allow project graphs that Include this file to specify their own defaults.
 	-->
@@ -77,6 +86,7 @@
 	<Option Name="DefaultCookedEditorConfigurations" DefaultValue="Development" Description="Default Cooked Editor configurations"/>
 	<Option Name="DefaultEditorTestList" DefaultValue="" Description="Default Editor test list"/>
 	<Option Name="DefaultTargetTestList" DefaultValue="" Description="Default Target test list"/>
+	<Option Name="DefaultZenStreamingTargetTestList" DefaultValue="" Description="Default Zen Streaming Target test list"/>
 	<Option Name="DefaultCookedEditorTestList" DefaultValue="" Description="Default Cooked Editor test list"/>
 	<Option Name="NetworkTempRootOverride" DefaultValue="" Description="Network temp root override"/>
 	<Option Name="NetworkPublishRootOverride" DefaultValue="" Description="Network publish root override"/>
@@ -126,6 +136,7 @@
 
 	<Option Name="EditorTestList" DefaultValue="$(DefaultEditorTestList)" Description="Tests to run on the editor. Test1+Test2 etc"/>
 	<Option Name="TargetTestList" DefaultValue="$(DefaultTargetTestList)" Description="Tests to run on the target platform(s). Test1+Test2 etc"/>
+	<Option Name="ZenStreamingTargetTestList" DefaultValue="$(DefaultZenStreamingTargetTestList)" Description="Zen Streaming tests to run on the target platform(s). Test1+Test2 etc"/>
 	<Option Name="CookedEditorTestList" DefaultValue="$(DefaultCookedEditorTestList)" Description="Tests to run on the cooked editor. Test1+Test2 etc"/>
 	<Option Name="EditorPlatformsToSkipBootTests" DefaultValue="$(EditorPlatformsToSkipTests)" Description="Editor platforms that should skip boot tests."/>
 	<Option Name="TargetPlatformsToSkipBootTests" DefaultValue="$(TargetPlatformsToSkipTests)" Description="Target platforms that should skip boot tests."/>
@@ -735,6 +746,7 @@
 		<Expand Name="CookAdditionalRequirements"/>
 
 		<Property Name="PlatformEditorNodeName" Value="" If="$(SkipEditorBuild)" />
+		<Property Name="PlatformSnapshotPublishNodeName" Value="$(PreNodeName)Snapshot Publish $(TargetPlatform)" />
 		<!-- Separate agent for cooking / staging, so the cook can happen in parallel with the client compile -->
 		<Agent Name="Package $(PlatformAgentName)" Type="$(PackageAgentType)">
 
@@ -805,7 +817,6 @@
 			<Property Name="ShouldSnapshotPublish" Value="false"/>
 			<Property Name="ShouldSnapshotPublish" Value="true" If="$(IsBuildMachine) and $(ShouldSnapshotPublishCloud) or $(ShouldSnapshotPublishZen)"/>
 			<Do If="$(PublishBuilds) and $(ShouldSnapshotPublish)">
-				<Property Name="PlatformSnapshotPublishNodeName" Value="$(PreNodeName)Snapshot Publish $(TargetPlatform)" />
 				<Node Name="$(PlatformSnapshotPublishNodeName)" Requires="$(PlatformHashNodeName)">
 					<Property Name="SnapshotLocalDir" Value="$(RootDir)/Snapshots" />
 					<Property Name="SnapshotIdentifier" Value="$(EscapedBranch)/$(BuildNamePath)"/>
@@ -1009,6 +1020,43 @@
 			</Node>
 			
 		</Agent>
+
+
+		<!-- Zen Streaming Import -->
+			<!-- Only Win64 supported at this time. -->
+		<Agent Name="Compile $(PlatformAgentName) Zen" Type="$(CompileAgentType)" If="'$(TargetPlatform)' == 'Win64'">
+			<Property Name="SnapshotFilenamePrefix" Value="SnapshotDescriptor"/>
+			<Property Name="SnapshotFilenamePrefix" Value="$(UE_HORDE_TEMPLATEID)" If="'$(UE_HORDE_TEMPLATEID)' != ''"/>
+			<Property Name="SnapshotLocalDir" Value="$(RootDir)/Snapshots" />			
+			<Property Name="SnapshotIdentifier" Value="$(UE_HORDE_TEMPLATEID).$(BuildVersion).$(UE_HORDE_JOBID:-6).$(UE_HORDE_STEPID)" If="'$(UE_HORDE_JOBID)' != '' and '$(UE_HORDE_STEPID)' != '' and '$(UE_HORDE_TEMPLATEID)' != ''"/>
+			<Property Name="SnapshotIdentifier" Value="$(ProjectName).$(UE_HORDE_JOBID)"/>
+			
+			<Property Name="ZenSnapshotDescriptorFile" Value="$(NetworkSnapshotDescFileRoot)/$(uebp_BuildRoot_Escaped)/$(ProjectName)/$(StagedPlatformFolder)/$(Change)/$(SnapshotFilenamePrefix)-builds.json" />
+			<Property Name="ZenBuildOutputDirectory" Value="$(ProjectOutputDirectory)" />
+			<Property Name="ZenPublishFolder" Value="$(RootDir)/$(ProjectPath)/Saved/StagedBuilds/$(StagedPlatformFolder)"/>
+			<Property Name="ZenPublishFolder" Value="$(ZenBuildOutputDirectory)/$(StagedPlatformFolder)" />
+ 
+			<Property Name="ZenArtificatNodeName" Value="$(CloudPlatformStagedArtifactNodeName)" If="$(PublishCloudArtifacts)"/>
+			<Property Name="ZenArtificatNodeName" Value="$(PlatformStagedArtifactNodeName)" If="$(PublishHordeArtifacts)"/>
+ 
+			<!-- Build Project without cooking -->
+			<Node Name="Zen Compile $(ProjectName)Game $(TargetPlatform)" Requires="$(ZenArtificatNodeName)">
+				<Log Message="ZenPublishFolder: $(ZenPublishFolder)"/>
+				<Compile Target="$(ProjectName)Game" Platform="$(TargetPlatform)" Configuration="Test" Arguments="-uba"/>
+			</Node>
+ 
+			<!-- Zen Streaming Stage -->
+			<Node Name="Zen Stage $(ProjectName) $(TargetPlatform)" Requires="Zen Compile $(ProjectName)Game $(TargetPlatform)">
+				<Command Name="BuildCookRun" Arguments="-project=$(ProjectName) -ScriptsForProject=$(ProjectName) -Platform=$(TargetPlatform) -Configuration=Test -SkipBuild -skipCook  -Stage -NoCodeSign " />
+			</Node>
+ 
+			<!-- Zen Streaming Publish -->
+			<Node Name="Zen Publish $(ProjectName) $(TargetPlatform) files" Requires="Zen Stage $(ProjectName) $(TargetPlatform)" Produces="#Published Zen $(ProjectName) $(TargetPlatform) Files">
+				<Log Message="Tagging files at: $(ZenPublishFolder)" />
+				<Tag Files="$(ZenPublishFolder)/..." With="#Published Zen $(ProjectName) $(TargetPlatform) Files"/>
+			</Node>
+			
+		</Agent>
 		
 		<!-- Test Time! -->
 		<Property Name="AllPlatformInstallNodes" Value="" />
@@ -1197,6 +1245,86 @@
 						<Expand Name="AdditionalTestNodes"/>
 					</ForEach>
 				</Agent>
+
+				<!-- Zen Streaming Import Tests -->
+					<!-- Only Win64 Zen streaming testing supported at this time. -->
+				<Property Name="TestTargetAgentType" Value="$(TestAgentType)"/>
+				<Property Name="TestTargetAgentType" Value="$(ConfigAgentType)" If="'$(ConfigAgentType)' != ''"/>
+				<Property Name="TestTargetAgentType" Value="$(ZenStreamingWinTestPool)" If="'$(TargetPlatform)' == 'Win64' and '$(ZenStreamingWinTestPool)' != ''" /> 
+ 
+				<Agent Name="Run Zen Streaming Tests $(TargetPlatform)" Type="$(TestAgentType)">
+					<!-- Setup Zen Args -->
+					<Property Name="SnapshotFilenamePrefix" Value="SnapshotDescriptor"/>
+					<Property Name="SnapshotFilenamePrefix" Value="$(UE_HORDE_TEMPLATEID)" If="'$(UE_HORDE_TEMPLATEID)' != ''"/>
+					<Property Name="SnapshotIdentifier" Value="$(UE_HORDE_TEMPLATEID).$(BuildVersion).$(UE_HORDE_JOBID:-6).$(UE_HORDE_STEPID)" If="'$(UE_HORDE_JOBID)' != '' and '$(UE_HORDE_STEPID)' != '' and '$(UE_HORDE_TEMPLATEID)' != ''"/>
+					<Property Name="SnapshotIdentifier" Value="$(ProjectName).$(UE_HORDE_JOBID)"/>
+					<Property Name="ZenSnapshotDescriptorFile" Value="$(NetworkSnapshotDescFileRoot)/$(uebp_BuildRoot_Escaped)/$(ProjectName)/$(StagedPlatformFolder)/$(Change)/$(SnapshotFilenamePrefix)-builds.json" />
+					
+					<!-- Setup Gauntlet Args -->
+					<Property Name="GauntletBuildArg" Value="-build=&quot;$(ProjectOutputDirectory)/$(StagedPlatformFolder)&quot;" />
+					<Property Name="GauntletBuildArg" Value="-build=&quot;$(ProjectOutputDirectory)/Packages/$(StagedPlatformFolder)&quot; -packaged" If="$(UsePackaged)" />
+ 
+					<Do If="!$(Horde) And ('$(NetworkOutputDirectory)' != '' Or '$(GeneratedBuildVersion)' != '')">
+						<Property Name="TestDirectory" Value="$(NetworkOutputDirectory)" />
+						<Property Name="TestDirectory" Value="$(GeneratedBuildDir)" If="'$(GeneratedBuildVersion)' != ''"/>
+						<Property Name="GauntletBuildArg" Value="-build=&quot;$(TestDirectory)&quot;" />
+						<Property Name="GauntletBuildArg" Value="$(GauntletBuildArg) -packaged" If="$(UsePackaged)" />
+						<Property Name="GauntletBuildArg" Value="$(GauntletBuildArg) -destlocalinstalldir=&quot;$(LocalInstallDirectory)&quot;" If="!$(IsLocalPlatform) And $(WithDeviceReservation)"/>
+					</Do>
+					
+					<!-- Run through all Test Configsconfigs -->
+					<ForEach Name="TestConfiguration" Values="$(TestConfigurations)" Separator="+">
+						<Property Name="NodeAnnotations" Value=""/>
+						<!-- Prerequisites for the test to run -->
+						<ForEach Name="TestName" Values="$(ZenStreamingTargetTestList)" Separator="+">
+							<Property Name="TestNameTitle" Value="$(TestName)"/>
+							<Expand Name="SanitizeTestNameTitle"/>
+							<!-- Run and report on the test-->
+							<Do If="'$(AutomationReportOutputDirectory)' != ''">
+								<Property Name="OutputFragmentPath" Value="$(TargetPlatform)/$(TestConfiguration)/$(TestNameTitle)" />
+								<Property Name="OutputReportPath" Value="$(AutomationReportOutputDirectory)/$(OutputFragmentPath)" />
+								<Property Name="GauntletReportArgs" Value="-logdir=&quot;$(OutputReportPath)&quot; -artifactname=Gauntlet -ReportExportPath=&quot;$(OutputReportPath)&quot;" />
+								<Property Name="GauntletReportArgs" Value="$(GauntletReportArgs) -ReportURL=&quot;$(AutomationReportUrl)/$(OutputFragmentPath)/&quot;" If="'$(AutomationReportUrl)' != ''" />
+							</Do>
+							
+							<Property Name="TestNodeName" Value="$(PreNodeName)$(TargetPlatform) $(TestConfiguration)$(TestConfigTitle) $(TestNameTitle)" />
+							<Property Name="NodeAnnotations" Value=""/>
+							<Property Name="NodeAnnotations" Value="Workflow=$(AnnotationsTarget)" If="ContainsItem('$(AnnotationsPlatforms)', '$(TargetPlatform)', ';') and ContainsItem('$(AnnotationsTests)', '$(TestName)', ';')"/>
+ 
+							<Property Name="TestNameBeforeParentheses" Value="$(TestName)"/>
+							<Property Name="TestArgumentsWithinParentheses" Value=""/>
+							<Expand Name="SplitCombinedTestNameAndArgs"/>
+ 
+							<Property Name="GauntletTestArgs" Value="$(GauntletTargetTestArgs)"/>
+							<Property Name="GauntletTestArgs" Value="$(GauntletTargetTestAltArgs)" If="ContainsItem('$(AlternateDevicePoolPlatforms)', '$(TargetPlatform)', ';')"/>
+							<Property Name="GauntletTestArgs" Value="$(GauntletTestArgs) -adbkeys=&quot;$(AdbKeysPath)&quot;" If="'$(AdbKeysPath)' != '' and '$(TargetPlatform)' == 'Android'"/>
+ 
+							<!-- 'Requires' should be the package game and the publish folder is used; however some tests may require the editor ie: UE.TargetAutomation -->
+							<Node Name="$(TestNodeName)" Requires="Zen Publish $(ProjectName) $(TargetPlatform) files" Annotations="$(NodeAnnotations)" If="$(DoBuildZenStreaming) and '$(TargetPlatform)' == 'Win64'">
+								<Property Name="ZenHost" Value="$(ZenStoreHost)"/>
+								<Property Name="ZenHost" Value="127.0.0.1" If="$(DoLocalZenTest)"/>
+								<Property Name="ZenPlatform" Value="$(TargetPlatform)"/>
+								<Property Name="ZenPlatform" Value="Windows" If="'$(TargetPlatform)' == 'Win64'"/>
+								<Property Name="ZenUtilitySpawnPath" Value="$(RootDir)\Engine\Binaries\Win64\zen.exe" />
+								<Property Name="SnapshotIdentifier" Value="$(ProjectName).$(UE_HORDE_JOBID)"/>
+								<Property Name="ZenArgs" Value="-clientargs=&quot;-ZenStoreHost=$(ZenHost) -ZenStorePort=$(ZenStorePort) -ZenStoreProject=$(SnapshotIdentifier) -ZenStorePlatform=$(ZenPlatform)&quot;" />
+ 
+								<!-- Launch a local Zen Server, import the snapshot, execute test, cleanup snapshot --> 
+								<ZenLaunch Project="$(TargetProject)" If="$(DoLocalZenTest)" />
+								<ZenImportSnapshot Project="$(TargetProject)" RemoteZenHost="$(ZenHost)" RemoteZenPort="$(ZenStorePort)" SnapshotDescriptorFile="$(ZenSnapshotDescriptorFile)" DestinationProjectId="$(SnapshotIdentifier)" If="$(DoLocalZenTest)"/>
+								<Command Name="RunUnreal" Arguments="-project=$(TargetProject) -test=&quot;$(TestNameBeforeParentheses)&quot; -platform=$(TargetPlatform) -configuration=$(TestConfiguration) -target=$(TargetName) $(GauntletBuildArg) $(GauntletTestArgs) $(TestArgumentsWithinParentheses) $(GauntletReportArgs) $(TestConfigArgs) $(ZenArgs)"/>
+							</Node>
+
+							<!-- Add it to our list -->
+							<Do If="!ContainsItem('$(TargetPlatformsToSkipTests)', '$(TargetPlatform)', '+') And (!Contains('$(TestListToIncludePlatforms)', '^$(TestNameBeforeParentheses)|') Or ContainsItem('$(TestListToIncludePlatforms)', '$(TestNameBeforeParentheses)|$(TargetPlatform)', '^'))">
+								<Property Name="AllPlatformTestNodes" Value="$(AllPlatformTestNodes);$(TestNodeName)" If="$(DoBuildZenStreaming) and '$(TargetPlatform)' == 'Win64'"/>
+							</Do>
+						</ForEach>
+
+						<!-- Expand the tests here -->
+						<Expand Name="AdditionalTestNodes"/>
+					</ForEach>
+				</Agent>
 			</Do>
 		</ForEach>
 		
@@ -1215,6 +1343,7 @@
 		<Property Name="PlatformTasks" Value="" />
 		<Property Name="PlatformTasks" Value="$(PlatformCompileNodeName);$(UploadSymbolsNodeName)" If="!$(SkipBuild)" />
 		<Property Name="PlatformTasks" Value="$(PlatformCompileNodeName) Server" If="!$(SkipBuild) and $(DedicatedServer)" />
+		<Property Name="PlatformTasks" Value="$(PlatformTasks);Zen Publish $(ProjectName) Win64 files" If="$(DoBuildZenStreaming)" />
 		<Property Name="PlatformTasks" Value="$(PlatformTasks);$(PlatformCookNodeName)" If="!$(SkipCook) and !$(SkipPublish)" />
 		<Property Name="PlatformTasks" Value="$(PlatformTasks);$(PlatformCookNodeName) Server" If="!$(SkipCook) and !$(SkipPublish) and $(DedicatedServer)" />
 		<Property Name="PlatformTasks" Value="$(PlatformTasks);$(RequirePlatformStaged)" If="!$(SkipPublish)" />
diff --git a/Engine/Plugins/Performance/AutomatedPerfTesting/Build/Inc/AutomatedPerfTestCommonSettings.xml b/Engine/Plugins/Performance/AutomatedPerfTesting/Build/Inc/AutomatedPerfTestCommonSettings.xml
--- a/Engine/Plugins/Performance/AutomatedPerfTesting/Build/Inc/AutomatedPerfTestCommonSettings.xml
+++ b/Engine/Plugins/Performance/AutomatedPerfTesting/Build/Inc/AutomatedPerfTestCommonSettings.xml
@@ -56,9 +56,15 @@
 	<Option Name="DoLLM" DefaultValue="false" Restrict="true|false" Description="Whether to run LLM (memory) test." />
 	<Option Name="DoGPUPerf" DefaultValue="false" Restrict="true|false" Description="Whether to run GPUPerf (GPU performance) test." />
 	<Option Name="DoInsightsTrace" DefaultValue="false" Restrict="true|false" Description="Whether or not to capture Unreal Insights trace during the run." />
+	<Option Name="DoBuildZenStreaming" DefaultValue="true" Restrict ="true|false" Description="Whether or not to perform Zen Streaming tests" />
 	<Option Name="DoGPUReshape" DefaultValue="false" Restrict="true|false" Description="Whether to run GPU instrumentation test with GPU Reshape." />
 	<Option Name="GPUReshapeWorkspace" DefaultValue="BasicWorkspace" Description="User workspace configuration for GPU Reshape." />
 
+	<!-- Zen Arguments -->
+	<Option Name="ZenStoreHost" DefaultValue="" Description="The hostname or ip of the target zenserver. Do not include http://" />
+	<Option Name="ZenStorePort" DefaultValue="8558" Description="The port to use on the target zen server." />
+	<Property Name="ZenStoreProject" Value ="$(ProjectName).$(UE_HORDE_JOBID)"/>
+
 	<!-- Construct the common arguments for the performance test and report commands -->
 	<Property Name="AutomatedPerformanceCommonArgs" Value = "-AutomatedPerfTest.UseShippingInsights=$(UseShippingInsights)" />
 	<Property Name="AutomatedPerformanceCommonArgs" Value="$(AutomatedPerformanceCommonArgs) -AutomatedPerfTest.TestID=$(TestID)"/>
@@ -82,7 +88,7 @@
 
 	<Property Name="PerfTestCommonArgs"     Value="$(AutomatedPerformanceCommonArgs) -iterations=$(IterationsPerf)"/>
 	<Property Name="LLMTestCommonArgs"      Value="$(AutomatedPerformanceCommonArgs) -iterations=$(IterationsLLM)"/>
-	<Property Name="GPUPerfTestCommonArgs"      Value="$(AutomatedPerformanceCommonArgs) -iterations=$(IterationsGPUPerf)"/>
+	<Property Name="GPUPerfTestCommonArgs"  Value="$(AutomatedPerformanceCommonArgs) -iterations=$(IterationsGPUPerf)"/>
 	<Property Name="InsightsTestCommonArgs" Value="$(AutomatedPerformanceCommonArgs) -iterations=$(IterationsInsights)"/>
 	
 	<!-- Reshape configuration -->
@@ -118,6 +124,13 @@
 			<Property Name="TargetTestList" Value="$(TargetTestList)+AutomatedPerfTest.SequenceTest{-AutomatedPerfTest.SequencePerfTest.MapSequenceName=$(MapSequenceComboName) -AutomatedPerfTest.DoCSVProfiler -AutomatedPerfTest.DoGPUPerf $(GPUPerfTestCommonArgs)}[Sequence GPUPerf $(MapSequenceComboName)]" If="$(DoGPUPerf)" />
 			<Property Name="TargetTestList" Value="$(TargetTestList)+AutomatedPerfTest.SequenceTest{-AutomatedPerfTest.SequencePerfTest.MapSequenceName=$(MapSequenceComboName) -AutomatedPerfTest.DoCSVProfiler -AutomatedPerfTest.DoInsightsTrace $(InsightsTestCommonArgs)}[Sequence Insights $(MapSequenceComboName)]" If="$(DoInsightsTrace)"/>
 			<Property Name="TargetTestList" Value="$(TargetTestList)+AutomatedPerfTest.SequenceTest{-AutomatedPerfTest.SequencePerfTest.MapSequenceName=$(MapSequenceComboName) -AutomatedPerfTest.DoGPUReshape $(GPUReshapeTestCommonArgs)}[Sequence GPUReshape $(MapSequenceComboName)]" If="$(DoGPUReshape)"/>
+			<!-- Zen Tests -->
+			<Property Name="ZenTestArgs" Value="-AutomatedPerfTest.DoBuildZenStreaming"/>
+			<Property Name="ZenStreamingTargetTestList" Value="$(ZenStreamingTargetTestList)+AutomatedPerfTest.SequenceTest{-AutomatedPerfTest.SequencePerfTest.MapSequenceName=$(MapSequenceComboName) -AutomatedPerfTest.DoCSVProfiler $(PerfTestCommonArgs) $(ZenTestArgs)}[Zen Sequence Perf $(MapSequenceComboName)]" If="$(DoPerf) and $(DoBuildZenStreaming)" />
+			<Property Name="ZenStreamingTargetTestList" Value="$(ZenStreamingTargetTestList)+AutomatedPerfTest.SequenceTest{-AutomatedPerfTest.SequencePerfTest.MapSequenceName=$(MapSequenceComboName) -AutomatedPerfTest.DoCSVProfiler -AutomatedPerfTest.DoLLM $(LLMTestCommonArgs) $(ZenTestArgs)}[Zen Sequence LLM $(MapSequenceComboName)]" If="$(DoLLM) and $(DoBuildZenStreaming)"/>
+			<Property Name="ZenStreamingTargetTestList" Value="$(ZenStreamingTargetTestList)+AutomatedPerfTest.SequenceTest{-AutomatedPerfTest.SequencePerfTest.MapSequenceName=$(MapSequenceComboName) -AutomatedPerfTest.DoCSVProfiler -AutomatedPerfTest.DoGPUPerf $(GPUPerfTestCommonArgs) $(ZenTestArgs)}[Zen Sequence GPUPerf $(MapSequenceComboName)]" If="$(DoGPUPerf) and $(DoBuildZenStreaming)" />
+			<Property Name="ZenStreamingTargetTestList" Value="$(ZenStreamingTargetTestList)+AutomatedPerfTest.SequenceTest{-AutomatedPerfTest.SequencePerfTest.MapSequenceName=$(MapSequenceComboName) -AutomatedPerfTest.DoCSVProfiler -AutomatedPerfTest.DoInsightsTrace $(InsightsTestCommonArgs) $(ZenTestArgs)}[Zen Sequence Insights $(MapSequenceComboName)]" If="$(DoInsightsTrace) and $(DoBuildZenStreaming)"/>
+			<Property Name="ZenStreamingTargetTestList" Value="$(ZenStreamingTargetTestList)+AutomatedPerfTest.SequenceTest{-AutomatedPerfTest.SequencePerfTest.MapSequenceName=$(MapSequenceComboName) -AutomatedPerfTest.DoGPUReshape $(GPUReshapeTestCommonArgs) $(ZenTestArgs)}[Zen Sequence GPUReshape $(MapSequenceComboName)]" If="$(DoGPUReshape) and $(DoBuildZenStreaming)"/>
 		</ForEach>
 	</Do>
 
@@ -128,6 +141,13 @@
 			<Property Name="TargetTestList" Value="$(TargetTestList)+AutomatedPerfTest.StaticCameraTest{-AutomatedPerfTest.StaticCameraPerfTest.MapName=$(StaticCameraMapName) -AutomatedPerfTest.DoCSVProfiler -AutomatedPerfTest.DoGPUPerf $(LLMTestCommonArgs)}[Static Camera GPUPerf $(StaticCameraMapName)]" If="$(DoGPUPerf)"/>
 			<Property Name="TargetTestList" Value="$(TargetTestList)+AutomatedPerfTest.StaticCameraTest{-AutomatedPerfTest.StaticCameraPerfTest.MapName=$(StaticCameraMapName) -AutomatedPerfTest.DoCSVProfiler -AutomatedPerfTest.DoInsightsTrace $(InsightsTestCommonArgs)}[Static Camera Insights $(StaticCameraMapName)]" If="$(DoInsightsTrace)"/>
 			<Property Name="TargetTestList" Value="$(TargetTestList)+AutomatedPerfTest.StaticCameraTest{-AutomatedPerfTest.StaticCameraPerfTest.MapName=$(StaticCameraMapName) -AutomatedPerfTest.DoGPUReshape $(GPUReshapeTestCommonArgs)}[Static Camera GPUReshape $(StaticCameraMapName)]" If="$(DoGPUReshape)"/>
+			<!-- Zen Tests -->
+			<Property Name="ZenTestArgs" Value="-AutomatedPerfTest.DoBuildZenStreaming"/>
+			<Property Name="ZenStreamingTargetTestList" Value="$(ZenStreamingTargetTestList)+AutomatedPerfTest.StaticCameraTest{-AutomatedPerfTest.StaticCameraPerfTest.MapName=$(StaticCameraMapName) -AutomatedPerfTest.DoCSVProfiler $(PerfTestCommonArgs) $(ZenTestArgs)}[Zen Static Camera Perf $(StaticCameraMapName)]" If="$(DoPerf) and $(DoBuildZenStreaming)" />
+			<Property Name="ZenStreamingTargetTestList" Value="$(ZenStreamingTargetTestList)+AutomatedPerfTest.StaticCameraTest{-AutomatedPerfTest.StaticCameraPerfTest.MapName=$(StaticCameraMapName) -AutomatedPerfTest.DoCSVProfiler -AutomatedPerfTest.DoLLM $(LLMTestCommonArgs) $(ZenTestArgs)}[Static Camera LLM $(StaticCameraMapName)]" If="$(DoLLM) and $(DoBuildZenStreaming)"/>
+			<Property Name="ZenStreamingTargetTestList" Value="$(ZenStreamingTargetTestList)+AutomatedPerfTest.StaticCameraTest{-AutomatedPerfTest.StaticCameraPerfTest.MapName=$(StaticCameraMapName) -AutomatedPerfTest.DoCSVProfiler -AutomatedPerfTest.DoGPUPerf $(LLMTestCommonArgs) $(ZenTestArgs)}[Static Camera GPUPerf $(StaticCameraMapName)]" If="$(DoGPUPerf) and $(DoBuildZenStreaming)"/>
+			<Property Name="ZenStreamingTargetTestList" Value="$(ZenStreamingTargetTestList)+AutomatedPerfTest.StaticCameraTest{-AutomatedPerfTest.StaticCameraPerfTest.MapName=$(StaticCameraMapName) -AutomatedPerfTest.DoCSVProfiler -AutomatedPerfTest.DoInsightsTrace $(InsightsTestCommonArgs) $(ZenTestArgs)}[Static Camera Insights $(StaticCameraMapName)]" If="$(DoInsightsTrace) and $(DoBuildZenStreaming)"/>			
+			<Property Name="ZenStreamingTargetTestList" Value="$(ZenStreamingTargetTestList)+AutomatedPerfTest.StaticCameraTest{-AutomatedPerfTest.StaticCameraPerfTest.MapName=$(StaticCameraMapName) -AutomatedPerfTest.DoGPUReshape $(GPUReshapeTestCommonArgs) $(ZenTestArgs)}[Static Camera GPUReshape $(StaticCameraMapName)]" If="$(DoGPUReshape) and $(DoBuildZenStreaming)"/>
 		</ForEach>
 	</Do>
 
@@ -137,6 +157,13 @@
 		<Property Name="TargetTestList" Value="$(TargetTestList)+AutomatedPerfTest.MaterialTest{-AutomatedPerfTest.DoCSVProfiler -AutomatedPerfTest.DoGPUPerf $(LLMTestCommonArgs)}[Material GPUPerf]" If="$(DoGPUPerf)"/>
 		<Property Name="TargetTestList" Value="$(TargetTestList)+AutomatedPerfTest.MaterialTest{-AutomatedPerfTest.DoCSVProfiler -AutomatedPerfTest.DoInsightsTrace $(InsightsTestCommonArgs)}[Material Insights]" If="$(DoInsightsTrace)"/>
 		<Property Name="TargetTestList" Value="$(TargetTestList)+AutomatedPerfTest.MaterialTest{-AutomatedPerfTest.DoCSVProfiler -AutomatedPerfTest.DoGPUReshape $(GPUReshapeTestCommonArgs)}[Material GPUReshape]" If="$(DoGPUReshape)"/>
+		<!-- Zen Tests -->
+		<Property Name="ZenTestArgs" Value="-AutomatedPerfTest.DoBuildZenStreaming"/>
+		<Property Name="ZenStreamingTargetTestList" Value="$(ZenStreamingTargetTestList)+AutomatedPerfTest.MaterialTest{-AutomatedPerfTest.DoCSVProfiler $(PerfTestCommonArgs) $(ZenTestArgs)}[Zen Material Perf]" If="$(DoPerf) and $(DoBuildZenStreaming)" />
+		<Property Name="ZenStreamingTargetTestList" Value="$(ZenStreamingTargetTestList)+AutomatedPerfTest.MaterialTest{-AutomatedPerfTest.DoCSVProfiler -AutomatedPerfTest.DoLLM $(LLMTestCommonArgs) $(ZenTestArgs)}[Zen Material LLM]" If="$(DoLLM) and $(DoBuildZenStreaming)"/>
+		<Property Name="ZenStreamingTargetTestList" Value="$(ZenStreamingTargetTestList)+AutomatedPerfTest.MaterialTest{-AutomatedPerfTest.DoCSVProfiler -AutomatedPerfTest.DoGPUPerf $(LLMTestCommonArgs) $(ZenTestArgs)}[Zen Material GPUPerf]" If="$(DoGPUPerf) and $(DoBuildZenStreaming)"/>
+		<Property Name="ZenStreamingTargetTestList" Value="$(ZenStreamingTargetTestList)+AutomatedPerfTest.MaterialTest{-AutomatedPerfTest.DoCSVProfiler -AutomatedPerfTest.DoInsightsTrace $(InsightsTestCommonArgs) $(ZenTestArgs)}[Zen Material Insights]" If="$(DoInsightsTrace) and $(DoBuildZenStreaming)"/>
+		<Property Name="ZenStreamingTargetTestList" Value="$(ZenStreamingTargetTestList)+AutomatedPerfTest.MaterialTest{-AutomatedPerfTest.DoCSVProfiler -AutomatedPerfTest.DoGPUReshape $(GPUReshapeTestCommonArgs) $(ZenTestArgs)}[Zen Material GPUReshape]" If="$(DoGPUReshape) and $(DoBuildZenStreaming)"/>	
 	</Do>
 
 	<Do If="$(WithAPTReplayTest)">
@@ -145,6 +172,13 @@
 		<Property Name="TargetTestList" Value="$(TargetTestList)+AutomatedPerfTest.ReplayTest{-AutomatedPerfTest.ReplayPerfTest.ReplayName=&quot;$(ReplayName)&quot; -AutomatedPerfTest.DoCSVProfiler -AutomatedPerfTest.DoGPUPerf $(LLMTestCommonArgs)}[Replay GPUPerf]" If="$(DoGPUPerf)"/>
 		<Property Name="TargetTestList" Value="$(TargetTestList)+AutomatedPerfTest.ReplayTest{-AutomatedPerfTest.ReplayPerfTest.ReplayName=&quot;$(ReplayName)&quot; -AutomatedPerfTest.DoCSVProfiler -AutomatedPerfTest.DoInsightsTrace $(InsightsTestCommonArgs)}[Replay Insights]" If="$(DoInsightsTrace)"/>
 		<Property Name="TargetTestList" Value="$(TargetTestList)+AutomatedPerfTest.ReplayTest{-AutomatedPerfTest.ReplayPerfTest.ReplayName=&quot;$(ReplayName)&quot; -AutomatedPerfTest.DoGPUReshape $(GPUReshapeTestCommonArgs)}[Replay GPUReshape]" If="$(DoGPUReshape)"/>
+		<!-- Zen Tests -->
+		<Property Name="ZenTestArgs" Value="-AutomatedPerfTest.DoBuildZenStreaming"/>
+		<Property Name="ZenStreamingTargetTestList" Value="$(ZenStreamingTargetTestList)+AutomatedPerfTest.ReplayTest{-AutomatedPerfTest.ReplayPerfTest.ReplayName=&quot;$(ReplayName)&quot; -AutomatedPerfTest.DoCSVProfiler $(PerfTestCommonArgs) $(ZenTestArgs)}[Zen Replay Perf]" If="$(DoPerf) and $(DoBuildZenStreaming)" />
+		<Property Name="ZenStreamingTargetTestList" Value="$(ZenStreamingTargetTestList)+AutomatedPerfTest.ReplayTest{-AutomatedPerfTest.ReplayPerfTest.ReplayName=&quot;$(ReplayName)&quot; -AutomatedPerfTest.DoCSVProfiler -AutomatedPerfTest.DoLLM $(LLMTestCommonArgs) $(ZenTestArgs)}[Zen Replay LLM]" If="$(DoLLM) and $(DoBuildZenStreaming)"/>
+		<Property Name="ZenStreamingTargetTestList" Value="$(ZenStreamingTargetTestList)+AutomatedPerfTest.ReplayTest{-AutomatedPerfTest.ReplayPerfTest.ReplayName=&quot;$(ReplayName)&quot; -AutomatedPerfTest.DoCSVProfiler -AutomatedPerfTest.DoGPUPerf $(LLMTestCommonArgs) $(ZenTestArgs)}[Zen Replay GPUPerf]" If="$(DoGPUPerf) and $(DoBuildZenStreaming)"/>
+		<Property Name="ZenStreamingTargetTestList" Value="$(ZenStreamingTargetTestList)+AutomatedPerfTest.ReplayTest{-AutomatedPerfTest.ReplayPerfTest.ReplayName=&quot;$(ReplayName)&quot; -AutomatedPerfTest.DoCSVProfiler -AutomatedPerfTest.DoInsightsTrace $(InsightsTestCommonArgs) $(ZenTestArgs)}[Zen Replay Insights]" If="$(DoInsightsTrace) and $(DoBuildZenStreaming)"/>
+		<Property Name="ZenStreamingTargetTestList" Value="$(ZenStreamingTargetTestList)+AutomatedPerfTest.ReplayTest{-AutomatedPerfTest.ReplayPerfTest.ReplayName=&quot;$(ReplayName)&quot; -AutomatedPerfTest.DoGPUReshape $(GPUReshapeTestCommonArgs) $(ZenTestArgs)}[Zen Replay GPUReshape]" If="$(DoGPUReshape) and $(DoBuildZenStreaming)"/>
 	</Do>
 
 </BuildGraph>
diff --git a/Engine/Plugins/Performance/AutomatedPerfTesting/Build/Scripts/AutomatedPerfTestConfig.cs b/Engine/Plugins/Performance/AutomatedPerfTesting/Build/Scripts/AutomatedPerfTestConfig.cs
--- a/Engine/Plugins/Performance/AutomatedPerfTesting/Build/Scripts/AutomatedPerfTestConfig.cs
+++ b/Engine/Plugins/Performance/AutomatedPerfTesting/Build/Scripts/AutomatedPerfTestConfig.cs
@@ -116,6 +116,12 @@ namespace AutomatedPerfTest
 		[AutoParamWithNames("BasicWorkspace", "AutomatedPerfTest.GPUReshapeWorkspace")]
 		public string GPUReshapeWorkspace;
 
+		/// <summary>
+		/// Whether to run Zen Streaming Import tests with -dobuildzenstreaming
+		/// </summary>
+		[AutoParamWithNames(false, "AutomatedPerfTest.DoBuildZenStreaming")]
+		public bool DoBuildZenStreaming;
+
 		/// <summary>
 		/// Whether to lock dynamic resolution. This is applied automatically when GPUPerf test is run.
 		/// </summary>
diff --git a/Engine/Plugins/Performance/AutomatedPerfTesting/Build/Scripts/AutomatedPerfTestNode.cs b/Engine/Plugins/Performance/AutomatedPerfTesting/Build/Scripts/AutomatedPerfTestNode.cs
--- a/Engine/Plugins/Performance/AutomatedPerfTesting/Build/Scripts/AutomatedPerfTestNode.cs
+++ b/Engine/Plugins/Performance/AutomatedPerfTesting/Build/Scripts/AutomatedPerfTestNode.cs
@@ -511,10 +511,26 @@ namespace AutomatedPerfTest
 				IsPRSAdded = true;
 			}
 
-			if(Report != null)
+			if (Report != null)
 			{
 				Report.SetMetadata("PerfType", GetSubtestName(Config));
-				Report.SetMetadata("PerfTest", string.IsNullOrEmpty(Config.TestID)? "Default" : Config.TestID);
+				if (!string.IsNullOrEmpty(Config.TestID))
+				{
+					if (GetCachedConfiguration().DoBuildZenStreaming)
+					{
+						Log.Info($"TestID: {Config.TestID}");
+						Report.SetMetadata("PerfTest", "APTZen");
+					}
+					else
+					{
+						Report.SetMetadata("PerfTest", Config.TestID);
+					}
+				}
+				else
+				{
+					Report.SetMetadata("PerfTest", "Default");
+				}
+				
 			}
 
 			return Report;
@@ -522,6 +538,8 @@ namespace AutomatedPerfTest
 
 		/// <inheritdoc/>
 		protected override string HordeReportTestKey => base.HordeReportTestKey + (string.IsNullOrEmpty(GetCachedConfiguration().TestID) ? "" : $".{GetCachedConfiguration().TestID}");
+		
+		
 
 		/// <summary>
 		/// Get report type for current configuration
@@ -541,6 +559,11 @@ namespace AutomatedPerfTest
 				SummaryTableType = "autoPerfReportLlm";
 				HistoricalReportType = "autoPerfReportLlm";
 			}
+
+			if(GetCachedConfiguration().DoBuildZenStreaming)
+			{
+				// Using defaults for now, but added here as a stub for incoming Zen Streaming Summary View.
+			}
 		}
 
 		/// <summary>
@@ -879,7 +902,11 @@ namespace AutomatedPerfTest
 
 			Log.Info("AutomatedPerfTestNode<>.GetConfiguration(): Config.DoFPSChart={0}, Config.DoCSVProfiler={1}, Config.DoVideoCapture={2}, Config.DoInsightsTrace={3}, Config.DoLLM={4}", 
 				Config.DoFPSChart, Config.DoCSVProfiler, Config.DoVideoCapture, Config.DoInsightsTrace, Config.DoLLM);
-			
+
+			if (Config.DoBuildZenStreaming)
+			{
+				Config.TestID = "APTZen";
+			}
 			ClientRole.CommandLineParams.AddOrAppendParamValue("AutomatedPerfTest.TestID", Config.TestID);
 
 			if (Config.DeviceProfileOverride != String.Empty)
@@ -903,6 +930,11 @@ namespace AutomatedPerfTest
 				ClientRole.CommandLineParams.Add("llmcsv");
 			}
 
+			if (Config.DoBuildZenStreaming)
+			{
+				ClientRole.CommandLineParams.Add("DoBuildZenStreaming");
+			}
+
 			if (Config.WaitForAttach)
 			{
 				ClientRole.CommandLineParams.Add("waitforattach");
