Fix DoBuildZenStreaming defaulting true.

diff --git a/Engine/Build/Graph/Tasks/BuildAndTestProject.xml b/Engine/Build/Graph/Tasks/BuildAndTestProject.xml
--- a/Engine/Build/Graph/Tasks/BuildAndTestProject.xml
+++ b/Engine/Build/Graph/Tasks/BuildAndTestProject.xml
@@ -70,8 +70,8 @@
 	<Option Name="PreflightPublishOverride" DefaultValue="false" Restrict="true|false" Description="Set to True to override Preflight Publishing when verifying gated behaviors." />
 
 	<!-- Zen Streaming Config -->
-	<Option Name="DoBuildZenStreaming" DefaultValue="true" Restrict="true|false" Description="A flag to enable ZenStreaming Building, Cooking, and Testing"/>
-	<Option Name="DoLocalZenTest" DefaultValue="true" Restrict="true|false" Description="Used to flag whether or not to use a local Zen Server for testing performance."/>
+	<Option Name="DoBuildZenStreaming" DefaultValue="false" Restrict="true|false" Description="A flag to enable ZenStreaming Building, Cooking, and Testing"/>
+	<Option Name="DoLocalZenTest" DefaultValue="false" Restrict="true|false" Description="Used to flag whether or not to use a local Zen Server for testing performance."/>
 	<Option Name="ZenStoreHost" DefaultValue="" Description="The hostname or ip of the target zenserver. Do not include http://" />
 	<Option Name="ZenStorePort" DefaultValue="8558" Description="The port to use on the target zen server." />
 	<Option Name="ZenStreamingWinTestPool" DefaultValue="" Description="Override the default Win64 Zen Streaming Test Agent Pool." />
@@ -1287,7 +1287,7 @@
 								<Property Name="GauntletReportArgs" Value="$(GauntletReportArgs) -ReportURL=&quot;$(AutomationReportUrl)/$(OutputFragmentPath)/&quot;" If="'$(AutomationReportUrl)' != ''" />
 							</Do>
 							
-							<Property Name="TestNodeName" Value="$(PreNodeName)$(TargetPlatform) $(TestConfiguration)$(TestConfigTitle) $(TestNameTitle)" />
+							<Property Name="TestNodeName" Value="ZenStreaming $(PreNodeName)$(TargetPlatform) $(TestConfiguration)$(TestConfigTitle) $(TestNameTitle)" />
 							<Property Name="NodeAnnotations" Value=""/>
 							<Property Name="NodeAnnotations" Value="Workflow=$(AnnotationsTarget)" If="ContainsItem('$(AnnotationsPlatforms)', '$(TargetPlatform)', ';') and ContainsItem('$(AnnotationsTests)', '$(TestName)', ';')"/>
  
