Add missing SetVersion invocations to IncrementalCompare and IncrementalValidate to fix issue where a plugin that is dependent on a branch name (part of the version information) was not being included in the build and this altered the global core redirect hash.

diff --git a/Engine/Build/Graph/Tests/IncrementalCompare.xml b/Engine/Build/Graph/Tests/IncrementalCompare.xml
--- a/Engine/Build/Graph/Tests/IncrementalCompare.xml
+++ b/Engine/Build/Graph/Tests/IncrementalCompare.xml
@@ -98,7 +98,11 @@
 	</Macro>
 
 	<Agent Name="IncrementalCompareAgent" Type="$(AgentType)">
-		<Node Name="$(CompileNodeName)">
+		<!-- Update the Version.h and Build.version files if we're making a formal build on a build machine -->
+		<Node Name="Update Version Files $(EditorPlatform)">
+			<SetVersion Change="$(Change)" Branch="$(EscapedBranch)" If="$(IsBuildMachine)"/>
+		</Node>
+		<Node Name="$(CompileNodeName)" Requires="Update Version Files $(EditorPlatform)">
 			<!-- Compile. -->
 			<Log Message="$(NewLine)========== IncrementalCompare.xml: Compile =========="/>
 			<Compile Target="ShaderCompileWorker" Platform="$(EditorPlatform)" Project="$(RelativeProjectFile)" Configuration="Development" Arguments="$(UBACompileArgs)"/>
diff --git a/Engine/Build/Graph/Tests/IncrementalValidate.xml b/Engine/Build/Graph/Tests/IncrementalValidate.xml
--- a/Engine/Build/Graph/Tests/IncrementalValidate.xml
+++ b/Engine/Build/Graph/Tests/IncrementalValidate.xml
@@ -67,7 +67,11 @@
 	<Property Name="IncrementalValidatePhase2Name" Value="Cook IncrementalValidatePhase2 $(ProjectName) $(TargetPlatform)"/>
 
 	<Agent Name="IncrementalValidateAgent" Type="$(AgentType)">
-		<Node Name="$(CompileNodeName)">
+		<!-- Update the Version.h and Build.version files if we're making a formal build on a build machine -->
+		<Node Name="Update Version Files $(EditorPlatform)">
+			<SetVersion Change="$(Change)" Branch="$(EscapedBranch)" If="$(IsBuildMachine)"/>
+		</Node>
+		<Node Name="$(CompileNodeName)" Requires="Update Version Files $(EditorPlatform)">
 			<Compile Target="ShaderCompileWorker" Platform="$(EditorPlatform)" Project="$(ProjectFile)" Configuration="Development" Arguments="$(UBACompileArgs)"/>
 			<Compile Target="$(EditorTarget)" Platform="$(EditorPlatform)" Project="$(ProjectFile)" Configuration="Development" Arguments="$(UBACompileArgs)"/>
 		</Node>
