Remove duplication of AdditionalTestingNodes expand usage in Zen Streaming tests

diff --git a/Engine/Build/Graph/Tasks/BuildAndTestProject.xml b/Engine/Build/Graph/Tasks/BuildAndTestProject.xml
--- a/Engine/Build/Graph/Tasks/BuildAndTestProject.xml
+++ b/Engine/Build/Graph/Tasks/BuildAndTestProject.xml
@@ -1320,9 +1320,6 @@
 								<Property Name="AllPlatformTestNodes" Value="$(AllPlatformTestNodes);$(TestNodeName)" If="$(DoBuildZenStreaming) and '$(TargetPlatform)' == 'Win64'"/>
 							</Do>
 						</ForEach>
-
-						<!-- Expand the tests here -->
-						<Expand Name="AdditionalTestNodes"/>
 					</ForEach>
 				</Agent>
 			</Do>
