APT Zen Streaming Multi-platform support.

Reworked and simplified the BuildAndTestProject logic for Zen Streaming staging, publishing, retrieval, installation, and test execution to allow for executing on all supported devices.

#Synced-CL 48498369

diff --git a/Engine/Build/Graph/Tasks/BuildAndTestProject.xml b/Engine/Build/Graph/Tasks/BuildAndTestProject.xml
--- a/Engine/Build/Graph/Tasks/BuildAndTestProject.xml
+++ b/Engine/Build/Graph/Tasks/BuildAndTestProject.xml
@@ -75,6 +75,7 @@
 	<Option Name="ZenStoreHost" DefaultValue="" Description="The hostname or ip of the target zenserver. Do not include http://" />
 	<Option Name="ZenStorePort" DefaultValue="8558" Description="The port to use on the target zen server." />
 	<Option Name="ZenStreamingWinTestPool" DefaultValue="" Description="Override the default Win64 Zen Streaming Test Agent Pool." />
+	<Option Name="RemoteZenHost" DefaultValue="" Description="The hostname or ip of the target zenserver to import remote build snapshots" />
 	
 	<!-- 
 		Allow project graphs that Include this file to specify their own defaults.
@@ -1021,41 +1022,51 @@
 			
 		</Agent>
 
+		<!-- Zen Streaming Import, Stage, Publish -->
+		<Agent Name="Stage and Publish $(PlatformAgentName) Zen" Type="$(CompileAgentType)">
+			<Node Name="Zen Stage $(ProjectName) $(TargetPlatform)" Requires="$(PlatformCompileNodeName);$(PlatformSnapshotPublishNodeName)">
+				<Property Name="SnapshotFilenamePrefix" Value="SnapshotDescriptor"/>
+				<Property Name="SnapshotFilenamePrefix" Value="$(UE_HORDE_TEMPLATEID)" If="'$(UE_HORDE_TEMPLATEID)' != ''"/>
+				<Property Name="SnapshotIdentifier" Value="$(ProjectName).$(UE_HORDE_JOBID)"/>			
+				
+				<Property Name="ZenSnapshotDescriptorFile" Value="$(NetworkSnapshotDescFileRoot)/$(uebp_BuildRoot_Escaped)/$(ProjectName)/$(StagedPlatformFolder)/$(Change)/$(SnapshotFilenamePrefix)-builds.json" />
 
-		<!-- Zen Streaming Import -->
-			<!-- Only Win64 supported at this time. -->
-		<Agent Name="Compile $(PlatformAgentName) Zen" Type="$(CompileAgentType)" If="'$(TargetPlatform)' == 'Win64'">
-			<Property Name="SnapshotFilenamePrefix" Value="SnapshotDescriptor"/>
-			<Property Name="SnapshotFilenamePrefix" Value="$(UE_HORDE_TEMPLATEID)" If="'$(UE_HORDE_TEMPLATEID)' != ''"/>
-			<Property Name="SnapshotLocalDir" Value="$(RootDir)/Snapshots" />			
-			<Property Name="SnapshotIdentifier" Value="$(UE_HORDE_TEMPLATEID).$(BuildVersion).$(UE_HORDE_JOBID:-6).$(UE_HORDE_STEPID)" If="'$(UE_HORDE_JOBID)' != '' and '$(UE_HORDE_STEPID)' != '' and '$(UE_HORDE_TEMPLATEID)' != ''"/>
-			<Property Name="SnapshotIdentifier" Value="$(ProjectName).$(UE_HORDE_JOBID)"/>
-			
-			<Property Name="ZenSnapshotDescriptorFile" Value="$(NetworkSnapshotDescFileRoot)/$(uebp_BuildRoot_Escaped)/$(ProjectName)/$(StagedPlatformFolder)/$(Change)/$(SnapshotFilenamePrefix)-builds.json" />
-			<Property Name="ZenBuildOutputDirectory" Value="$(ProjectOutputDirectory)" />
-			<Property Name="ZenPublishFolder" Value="$(RootDir)/$(ProjectPath)/Saved/StagedBuilds/$(StagedPlatformFolder)"/>
-			<Property Name="ZenPublishFolder" Value="$(ZenBuildOutputDirectory)/$(StagedPlatformFolder)" />
- 
-			<Property Name="ZenArtificatNodeName" Value="$(CloudPlatformStagedArtifactNodeName)" If="$(PublishCloudArtifacts)"/>
-			<Property Name="ZenArtificatNodeName" Value="$(PlatformStagedArtifactNodeName)" If="$(PublishHordeArtifacts)"/>
- 
-			<!-- Build Project without cooking -->
-			<Node Name="Zen Compile $(ProjectName)Game $(TargetPlatform)" Requires="$(ZenArtificatNodeName)">
-				<Log Message="ZenPublishFolder: $(ZenPublishFolder)"/>
-				<Compile Target="$(ProjectName)Game" Platform="$(TargetPlatform)" Configuration="Test" Arguments="-uba"/>
-			</Node>
- 
-			<!-- Zen Streaming Stage -->
-			<Node Name="Zen Stage $(ProjectName) $(TargetPlatform)" Requires="Zen Compile $(ProjectName)Game $(TargetPlatform)">
-				<Command Name="BuildCookRun" Arguments="-project=$(ProjectName) -ScriptsForProject=$(ProjectName) -Platform=$(TargetPlatform) -Configuration=Test -SkipBuild -skipCook  -Stage -NoCodeSign " />
+				<Property Name="ZenBuildArgs" Value="-project=$(TargetProject) -platform=$(TargetPlatform) -ScriptsForProject=$(ProjectName) $(StageAndPackageArguments)" />
+				<Property Name="ZenBuildArgs" Value="$(ZenBuildArgs) -client" If="'$(TargetType)' == 'Client'" />			
+
+				<ZenImportSnapshot Project="$(TargetProject)" SnapshotDescriptorFile="$(ZenSnapshotDescriptorFile)" DestinationProjectId="$(SnapshotIdentifier)" If="$(DoLocalZenTest)"/>
+
+				<!-- Stage -->
+				<Switch>
+					<Case If="$(IsMobilePlatform)">
+						<ForEach Name="PackageConfiguration" Values="$(PackageConfigurations)" Separator="+">
+							<Command Name="BuildCookRun" Arguments="$(ZenBuildArgs) -configuration=$(PackageConfiguration) -skipbuild -skipcook -stage -pak" />
+						</ForEach>
+					</Case>
+					<Default>
+						<Property Name="ClientStageArguments" Value="$(ZenBuildArgs) -configuration=$(PackageConfigurations) -skipbuild -skipcook -stage -pak -NoCodeSign" />
+						<Command Name="BuildCookRun" Arguments="$(ClientStageArguments)" />
+					</Default>
+				</Switch>
 			</Node>
  
 			<!-- Zen Streaming Publish -->
-			<Node Name="Zen Publish $(ProjectName) $(TargetPlatform) files" Requires="Zen Stage $(ProjectName) $(TargetPlatform)" Produces="#Published Zen $(ProjectName) $(TargetPlatform) Files">
-				<Log Message="Tagging files at: $(ZenPublishFolder)" />
-				<Tag Files="$(ZenPublishFolder)/..." With="#Published Zen $(ProjectName) $(TargetPlatform) Files"/>
-			</Node>
-			
+			<Node Name="Zen Publish $(ProjectName) $(TargetPlatform) files" Requires="Zen Stage $(ProjectName) $(TargetPlatform)" >
+				<Property Name="TargetDirectory" Value="$(ProjectOutputDirectory)/$(StagedPlatformFolder)"/>
+				<Expand Name="WriteArtifactMetadataFile" Type="staged-build" Platform="$(StagedPlatformFolder)" Configuration="$(PackageConfigurations)" ArtifactDir="$(TargetDirectory)" />
+				<Tag Files="$(TargetDirectory)/..." With="#Published Zen $(ProjectName) $(TargetPlatform) Files"/>
+
+				<CreateCloudArtifact 
+					Name="$(ProjectName)-$(TargetPlatform)-Zen-Staged"
+					BuildGroup="$(BuildVersion)" 
+					Type="staged-build" 
+					BaseDir="$(TargetDirectory)" 
+					Metadata="cookPlatform=$(StagedPlatformFolder);dashboard-category=Staged Builds - Cloud;$(SharedCloudBuildMetadata)" 
+					Project="$(TargetProject)" 
+					Platform="$(StagedPlatformFolder)" 
+					Files="#Published Zen $(ProjectName) $(TargetPlatform) Files"/>
+				<Delete Files="$(TargetDirectory)/artifactmetadata.json"/>
+			</Node>	 
 		</Agent>
 		
 		<!-- Test Time! -->
@@ -1247,22 +1258,22 @@
 				</Agent>
 
 				<!-- Zen Streaming Import Tests -->
-					<!-- Only Win64 Zen streaming testing supported at this time. -->
 				<Property Name="TestTargetAgentType" Value="$(TestAgentType)"/>
 				<Property Name="TestTargetAgentType" Value="$(ConfigAgentType)" If="'$(ConfigAgentType)' != ''"/>
 				<Property Name="TestTargetAgentType" Value="$(ZenStreamingWinTestPool)" If="'$(TargetPlatform)' == 'Win64' and '$(ZenStreamingWinTestPool)' != ''" /> 
  
-				<Agent Name="Run Zen Streaming Tests $(TargetPlatform)" Type="$(TestAgentType)">
+				<Agent Name="Run Zen Streaming Tests $(TargetPlatform)" Type="$(TestTargetAgentType)">
 					<!-- Setup Zen Args -->
 					<Property Name="SnapshotFilenamePrefix" Value="SnapshotDescriptor"/>
 					<Property Name="SnapshotFilenamePrefix" Value="$(UE_HORDE_TEMPLATEID)" If="'$(UE_HORDE_TEMPLATEID)' != ''"/>
-					<Property Name="SnapshotIdentifier" Value="$(UE_HORDE_TEMPLATEID).$(BuildVersion).$(UE_HORDE_JOBID:-6).$(UE_HORDE_STEPID)" If="'$(UE_HORDE_JOBID)' != '' and '$(UE_HORDE_STEPID)' != '' and '$(UE_HORDE_TEMPLATEID)' != ''"/>
 					<Property Name="SnapshotIdentifier" Value="$(ProjectName).$(UE_HORDE_JOBID)"/>
 					<Property Name="ZenSnapshotDescriptorFile" Value="$(NetworkSnapshotDescFileRoot)/$(uebp_BuildRoot_Escaped)/$(ProjectName)/$(StagedPlatformFolder)/$(Change)/$(SnapshotFilenamePrefix)-builds.json" />
-					
+
 					<!-- Setup Gauntlet Args -->
 					<Property Name="GauntletBuildArg" Value="-build=&quot;$(ProjectOutputDirectory)/$(StagedPlatformFolder)&quot;" />
-					<Property Name="GauntletBuildArg" Value="-build=&quot;$(ProjectOutputDirectory)/Packages/$(StagedPlatformFolder)&quot; -packaged" If="$(UsePackaged)" />
+					<Property Name="GauntletTestArgs" Value="$(GauntletTargetTestArgs)"/>
+					<Property Name="GauntletTestArgs" Value="$(GauntletTargetTestAltArgs)" If="ContainsItem('$(AlternateDevicePoolPlatforms)', '$(TargetPlatform)', ';')"/>
+					<Property Name="GauntletTestArgs" Value="$(GauntletTestArgs) -adbkeys=&quot;$(AdbKeysPath)&quot;" If="'$(AdbKeysPath)' != '' and '$(TargetPlatform)' == 'Android'"/>
  
 					<Do If="!$(Horde) And ('$(NetworkOutputDirectory)' != '' Or '$(GeneratedBuildVersion)' != '')">
 						<Property Name="TestDirectory" Value="$(NetworkOutputDirectory)" />
@@ -1275,10 +1286,55 @@
 					<!-- Run through all Test Configsconfigs -->
 					<ForEach Name="TestConfiguration" Values="$(TestConfigurations)" Separator="+">
 						<Property Name="NodeAnnotations" Value=""/>
-						<!-- Prerequisites for the test to run -->
+						
+						<!-- Retrieve Staged Build-->
+						<Property Name="ZenTestPrerequisites" Value="Zen Publish $(ProjectName) $(TargetPlatform) files" />
+						<Property Name="ZenCloudArtifactPlatformRetrieveWithConfig" Value="Zen $(CloudArtifactPlatformRetrieve) $(TestConfiguration)$(TestConfigTitle)"/>
+						<Node Name="$(ZenCloudArtifactPlatformRetrieveWithConfig)" Requires="$(ZenTestPrerequisites)">
+							<RetrieveCloudArtifact 
+								Name="$(ProjectName)-$(TargetPlatform)-Zen-Staged" 
+								Type="staged-build" 
+								OutputDir="$(ProjectOutputDirectory)/$(StagedPlatformFolder)" 
+								Project="$(TargetProject)" 
+								Platform="$(StagedPlatformFolder)" 
+								BuildGroup="$(BuildVersion)" 
+								AllowMultipleMatches="true"
+							/>
+						</Node>
+						<Property Name="ZenTestPrerequisites" Value="$(ZenTestPrerequisites);$(ZenCloudArtifactPlatformRetrieveWithConfig)" />
+
+						<!-- Install step -->
+						<Do If="!$(IsLocalPlatform) And $(WithDeviceReservation)">
+							<Property Name="NodeAnnotations" Value=""/>
+							<Property Name="NodeAnnotations" Value="Workflow=$(AnnotationsTarget)" If="ContainsItem('$(AnnotationsPlatforms)', '$(TargetPlatform)', ';') and ContainsItem('$(AnnotationsTests)', 'Install', ';')"/>
+							<Property Name="ZenInstallNodeName" Value="Zen $(PreNodeName)$(TargetPlatform) $(TestConfiguration)$(TestConfigTitle) Install Game"/>
+							<Property Name="InstallArgs" Value="-ClientCount=1 $(GauntletPreInstallArgs)"/>
+
+							<Do If="'$(AutomationReportOutputDirectory)' != ''">
+								<Property Name="OutputFragmentPath" Value="$(TargetPlatform)/$(TestConfiguration)/UE.InstallOnly" />
+								<Property Name="OutputReportPath" Value="$(AutomationReportOutputDirectory)/$(OutputFragmentPath)" />
+								<Property Name="GauntletReportArgs" Value="-logdir=&quot;$(OutputReportPath)&quot; -artifactname=Gauntlet" />
+							</Do>
+
+							<Node Name="$(ZenInstallNodeName)" Requires="$(ZenTestPrerequisites)" Annotations="DeviceReserve=Begin;$(NodeAnnotations)">
+								<Command Name="RunUnreal" Arguments="-project=$(TargetProject) -test=&quot;UE.InstallOnly&quot; -platform=$(TargetPlatform) -configuration=$(TestConfiguration) -target=$(TargetName) $(GauntletBuildArg) $(GauntletTestArgs) $(GauntletReportArgs) $(InstallArgs)" />
+							</Node>
+
+							<Property Name="ZenTestPrerequisites" Value="$(ZenTestPrerequisites);$(ZenInstallNodeName)" />
+
+							<Do If="!ContainsItem('$(TargetPlatformsToSkipBootTests)', '$(TargetPlatform)', '+') Or !ContainsItem('$(TargetPlatformsToSkipTests)', '$(TargetPlatform)', '+')">
+								<Property Name="AllPlatformInstallNodes" Value="$(AllPlatformInstallNodes);$(ZenInstallNodeName)"/>
+							</Do>
+
+							<!-- Skip install for all other test node -->
+							<Property Name="GauntletTestArgs" Value="$(GauntletTestArgs) $(GauntletTestSkipInstallArgs)"/>
+						</Do>
+
+						<!-- Execute Zen Streaming Tests-->
 						<ForEach Name="TestName" Values="$(ZenStreamingTargetTestList)" Separator="+">
 							<Property Name="TestNameTitle" Value="$(TestName)"/>
 							<Expand Name="SanitizeTestNameTitle"/>
+							
 							<!-- Run and report on the test-->
 							<Do If="'$(AutomationReportOutputDirectory)' != ''">
 								<Property Name="OutputFragmentPath" Value="$(TargetPlatform)/$(TestConfiguration)/$(TestNameTitle)" />
@@ -1298,18 +1354,16 @@
 							<Property Name="GauntletTestArgs" Value="$(GauntletTargetTestArgs)"/>
 							<Property Name="GauntletTestArgs" Value="$(GauntletTargetTestAltArgs)" If="ContainsItem('$(AlternateDevicePoolPlatforms)', '$(TargetPlatform)', ';')"/>
 							<Property Name="GauntletTestArgs" Value="$(GauntletTestArgs) -adbkeys=&quot;$(AdbKeysPath)&quot;" If="'$(AdbKeysPath)' != '' and '$(TargetPlatform)' == 'Android'"/>
- 
+
 							<!-- 'Requires' should be the package game and the publish folder is used; however some tests may require the editor ie: UE.TargetAutomation -->
-							<Node Name="$(TestNodeName)" Requires="Zen Publish $(ProjectName) $(TargetPlatform) files" Annotations="$(NodeAnnotations)" If="$(DoBuildZenStreaming) and '$(TargetPlatform)' == 'Win64'">
+							<Node Name="$(TestNodeName)" Requires="$(PlatformToolsNodeName);$(ZenTestPrerequisites)" Annotations="$(NodeAnnotations)">
 								<Property Name="ZenHost" Value="$(ZenStoreHost)"/>
 								<Property Name="ZenHost" Value="127.0.0.1" If="$(DoLocalZenTest)"/>
 								<Property Name="ZenPlatform" Value="$(TargetPlatform)"/>
 								<Property Name="ZenPlatform" Value="Windows" If="'$(TargetPlatform)' == 'Win64'"/>
-								<Property Name="ZenUtilitySpawnPath" Value="$(RootDir)\Engine\Binaries\Win64\zen.exe" />
-								<Property Name="SnapshotIdentifier" Value="$(ProjectName).$(UE_HORDE_JOBID)"/>
-								<Property Name="ZenArgs" Value="-clientargs=&quot;-ZenStoreHost=$(ZenHost) -ZenStorePort=$(ZenStorePort) -ZenStoreProject=$(SnapshotIdentifier) -ZenStorePlatform=$(ZenPlatform)&quot;" />
+								<Property Name="ZenArgs" Value="-ZenStoreHost=$(ZenHost) -clientargs=&quot;-ZenStorePort=$(ZenStorePort) -ZenStoreProject=$(SnapshotIdentifier) -ZenStorePlatform=$(ZenPlatform)&quot;" />
  
-								<!-- Launch a local Zen Server, import the snapshot, execute test, cleanup snapshot --> 
+								<!-- Launch a local Zen Server, import the snapshot, execute test, cleanup snapshot -->
 								<ZenLaunch Project="$(TargetProject)" If="$(DoLocalZenTest)" />
 								<ZenImportSnapshot Project="$(TargetProject)" RemoteZenHost="$(ZenHost)" RemoteZenPort="$(ZenStorePort)" SnapshotDescriptorFile="$(ZenSnapshotDescriptorFile)" DestinationProjectId="$(SnapshotIdentifier)" If="$(DoLocalZenTest)"/>
 								<Command Name="RunUnreal" Arguments="-project=$(TargetProject) -test=&quot;$(TestNameBeforeParentheses)&quot; -platform=$(TargetPlatform) -configuration=$(TestConfiguration) -target=$(TargetName) $(GauntletBuildArg) $(GauntletTestArgs) $(TestArgumentsWithinParentheses) $(GauntletReportArgs) $(TestConfigArgs) $(ZenArgs)"/>
@@ -1317,7 +1371,7 @@
 
 							<!-- Add it to our list -->
 							<Do If="!ContainsItem('$(TargetPlatformsToSkipTests)', '$(TargetPlatform)', '+') And (!Contains('$(TestListToIncludePlatforms)', '^$(TestNameBeforeParentheses)|') Or ContainsItem('$(TestListToIncludePlatforms)', '$(TestNameBeforeParentheses)|$(TargetPlatform)', '^'))">
-								<Property Name="AllPlatformTestNodes" Value="$(AllPlatformTestNodes);$(TestNodeName)" If="$(DoBuildZenStreaming) and '$(TargetPlatform)' == 'Win64'"/>
+								<Property Name="AllPlatformTestNodes" Value="$(AllPlatformTestNodes);$(TestNodeName)" If="$(DoBuildZenStreaming)"/>
 							</Do>
 						</ForEach>
 					</ForEach>
@@ -1340,7 +1394,7 @@
 		<Property Name="PlatformTasks" Value="" />
 		<Property Name="PlatformTasks" Value="$(PlatformCompileNodeName);$(UploadSymbolsNodeName)" If="!$(SkipBuild)" />
 		<Property Name="PlatformTasks" Value="$(PlatformCompileNodeName) Server" If="!$(SkipBuild) and $(DedicatedServer)" />
-		<Property Name="PlatformTasks" Value="$(PlatformTasks);Zen Publish $(ProjectName) Win64 files" If="$(DoBuildZenStreaming)" />
+		<Property Name="PlatformTasks" Value="$(PlatformTasks);Zen Publish $(ProjectName) $(TargetPlatform) files" If="$(DoBuildZenStreaming)"/>
 		<Property Name="PlatformTasks" Value="$(PlatformTasks);$(PlatformCookNodeName)" If="!$(SkipCook) and !$(SkipPublish)" />
 		<Property Name="PlatformTasks" Value="$(PlatformTasks);$(PlatformCookNodeName) Server" If="!$(SkipCook) and !$(SkipPublish) and $(DedicatedServer)" />
 		<Property Name="PlatformTasks" Value="$(PlatformTasks);$(RequirePlatformStaged)" If="!$(SkipPublish)" />
diff --git a/Engine/Plugins/Performance/AutomatedPerfTesting/Build/Inc/AutomatedPerfTestCommonSettings.xml b/Engine/Plugins/Performance/AutomatedPerfTesting/Build/Inc/AutomatedPerfTestCommonSettings.xml
--- a/Engine/Plugins/Performance/AutomatedPerfTesting/Build/Inc/AutomatedPerfTestCommonSettings.xml
+++ b/Engine/Plugins/Performance/AutomatedPerfTesting/Build/Inc/AutomatedPerfTestCommonSettings.xml
@@ -57,15 +57,10 @@
 	<Option Name="DoLLM" DefaultValue="false" Restrict="true|false" Description="Whether to run LLM (memory) test." />
 	<Option Name="DoGPUPerf" DefaultValue="false" Restrict="true|false" Description="Whether to run GPUPerf (GPU performance) test." />
 	<Option Name="DoInsightsTrace" DefaultValue="false" Restrict="true|false" Description="Whether or not to capture Unreal Insights trace during the run." />
-	<Option Name="DoBuildZenStreaming" DefaultValue="true" Restrict ="true|false" Description="Whether or not to perform Zen Streaming tests" />
 	<Option Name="DoGPUReshape" DefaultValue="false" Restrict="true|false" Description="Whether to run GPU instrumentation test with GPU Reshape." />
 	<Option Name="GPUReshapeWorkspace" DefaultValue="BasicWorkspace" Description="User workspace configuration for GPU Reshape." />
 
-	<!-- Zen Arguments -->
-	<Option Name="ZenStoreHost" DefaultValue="" Description="The hostname or ip of the target zenserver. Do not include http://" />
-	<Option Name="ZenStorePort" DefaultValue="8558" Description="The port to use on the target zen server." />
-	<Property Name="ZenStoreProject" Value ="$(ProjectName).$(UE_HORDE_JOBID)"/>
-
+	
 	<!-- Construct the common arguments for the performance test and report commands -->
 	<Property Name="AutomatedPerformanceCommonArgs" Value = "-AutomatedPerfTest.UseShippingInsights=$(UseShippingInsights)" />
 	<Property Name="AutomatedPerformanceCommonArgs" Value="$(AutomatedPerformanceCommonArgs) -AutomatedPerfTest.TestID=$(TestID)"/>
@@ -91,6 +86,10 @@
 	<Property Name="LLMTestCommonArgs"      Value="$(AutomatedPerformanceCommonArgs) -iterations=$(IterationsLLM)"/>
 	<Property Name="GPUPerfTestCommonArgs"  Value="$(AutomatedPerformanceCommonArgs) -iterations=$(IterationsGPUPerf)"/>
 	<Property Name="InsightsTestCommonArgs" Value="$(AutomatedPerformanceCommonArgs) -iterations=$(IterationsInsights)"/>
+
+	<!-- Zen Arguments -->
+	<Property Name="ZenTestArgs" Value="-AutomatedPerfTest.DoBuildZenStreaming -AutomatedPerfTest.ZenStoreHost=&quot;$(ZenStoreHost)&quot;"/>
+	<Property Name="ZenTestArgs" Value="$(ZenTestArgs) -AutomatedPerfTest.DoLocalZenTest" If="$(DoLocalZenTest)" />
 	
 	<!-- Reshape configuration -->
 	<Do If="$(DoGPUReshape)">
@@ -126,7 +125,6 @@
 			<Property Name="TargetTestList" Value="$(TargetTestList)+AutomatedPerfTest.SequenceTest{-AutomatedPerfTest.SequencePerfTest.MapSequenceName=$(MapSequenceComboName) -AutomatedPerfTest.DoCSVProfiler -AutomatedPerfTest.DoInsightsTrace $(InsightsTestCommonArgs)}[Sequence Insights $(MapSequenceComboName)]" If="$(DoInsightsTrace)"/>
 			<Property Name="TargetTestList" Value="$(TargetTestList)+AutomatedPerfTest.SequenceTest{-AutomatedPerfTest.SequencePerfTest.MapSequenceName=$(MapSequenceComboName) -AutomatedPerfTest.DoGPUReshape $(GPUReshapeTestCommonArgs)}[Sequence GPUReshape $(MapSequenceComboName)]" If="$(DoGPUReshape)"/>
 			<!-- Zen Tests -->
-			<Property Name="ZenTestArgs" Value="-AutomatedPerfTest.DoBuildZenStreaming"/>
 			<Property Name="ZenStreamingTargetTestList" Value="$(ZenStreamingTargetTestList)+AutomatedPerfTest.SequenceTest{-AutomatedPerfTest.SequencePerfTest.MapSequenceName=$(MapSequenceComboName) -AutomatedPerfTest.DoCSVProfiler $(PerfTestCommonArgs) $(ZenTestArgs)}[Zen Sequence Perf $(MapSequenceComboName)]" If="$(DoPerf) and $(DoBuildZenStreaming)" />
 			<Property Name="ZenStreamingTargetTestList" Value="$(ZenStreamingTargetTestList)+AutomatedPerfTest.SequenceTest{-AutomatedPerfTest.SequencePerfTest.MapSequenceName=$(MapSequenceComboName) -AutomatedPerfTest.DoCSVProfiler -AutomatedPerfTest.DoLLM $(LLMTestCommonArgs) $(ZenTestArgs)}[Zen Sequence LLM $(MapSequenceComboName)]" If="$(DoLLM) and $(DoBuildZenStreaming)"/>
 			<Property Name="ZenStreamingTargetTestList" Value="$(ZenStreamingTargetTestList)+AutomatedPerfTest.SequenceTest{-AutomatedPerfTest.SequencePerfTest.MapSequenceName=$(MapSequenceComboName) -AutomatedPerfTest.DoCSVProfiler -AutomatedPerfTest.DoGPUPerf $(GPUPerfTestCommonArgs) $(ZenTestArgs)}[Zen Sequence GPUPerf $(MapSequenceComboName)]" If="$(DoGPUPerf) and $(DoBuildZenStreaming)" />
@@ -143,7 +141,6 @@
 			<Property Name="TargetTestList" Value="$(TargetTestList)+AutomatedPerfTest.StaticCameraTest{-AutomatedPerfTest.StaticCameraPerfTest.MapName=$(StaticCameraMapName) -AutomatedPerfTest.DoCSVProfiler -AutomatedPerfTest.DoInsightsTrace $(InsightsTestCommonArgs)}[Static Camera Insights $(StaticCameraMapName)]" If="$(DoInsightsTrace)"/>
 			<Property Name="TargetTestList" Value="$(TargetTestList)+AutomatedPerfTest.StaticCameraTest{-AutomatedPerfTest.StaticCameraPerfTest.MapName=$(StaticCameraMapName) -AutomatedPerfTest.DoGPUReshape $(GPUReshapeTestCommonArgs)}[Static Camera GPUReshape $(StaticCameraMapName)]" If="$(DoGPUReshape)"/>
 			<!-- Zen Tests -->
-			<Property Name="ZenTestArgs" Value="-AutomatedPerfTest.DoBuildZenStreaming"/>
 			<Property Name="ZenStreamingTargetTestList" Value="$(ZenStreamingTargetTestList)+AutomatedPerfTest.StaticCameraTest{-AutomatedPerfTest.StaticCameraPerfTest.MapName=$(StaticCameraMapName) -AutomatedPerfTest.DoCSVProfiler $(PerfTestCommonArgs) $(ZenTestArgs)}[Zen Static Camera Perf $(StaticCameraMapName)]" If="$(DoPerf) and $(DoBuildZenStreaming)" />
 			<Property Name="ZenStreamingTargetTestList" Value="$(ZenStreamingTargetTestList)+AutomatedPerfTest.StaticCameraTest{-AutomatedPerfTest.StaticCameraPerfTest.MapName=$(StaticCameraMapName) -AutomatedPerfTest.DoCSVProfiler -AutomatedPerfTest.DoLLM $(LLMTestCommonArgs) $(ZenTestArgs)}[Static Camera LLM $(StaticCameraMapName)]" If="$(DoLLM) and $(DoBuildZenStreaming)"/>
 			<Property Name="ZenStreamingTargetTestList" Value="$(ZenStreamingTargetTestList)+AutomatedPerfTest.StaticCameraTest{-AutomatedPerfTest.StaticCameraPerfTest.MapName=$(StaticCameraMapName) -AutomatedPerfTest.DoCSVProfiler -AutomatedPerfTest.DoGPUPerf $(LLMTestCommonArgs) $(ZenTestArgs)}[Static Camera GPUPerf $(StaticCameraMapName)]" If="$(DoGPUPerf) and $(DoBuildZenStreaming)"/>
@@ -158,8 +155,6 @@
 		<Property Name="TargetTestList" Value="$(TargetTestList)+AutomatedPerfTest.MaterialTest{-AutomatedPerfTest.DoCSVProfiler -AutomatedPerfTest.DoGPUPerf $(LLMTestCommonArgs)}[Material GPUPerf]" If="$(DoGPUPerf)"/>
 		<Property Name="TargetTestList" Value="$(TargetTestList)+AutomatedPerfTest.MaterialTest{-AutomatedPerfTest.DoCSVProfiler -AutomatedPerfTest.DoInsightsTrace $(InsightsTestCommonArgs)}[Material Insights]" If="$(DoInsightsTrace)"/>
 		<Property Name="TargetTestList" Value="$(TargetTestList)+AutomatedPerfTest.MaterialTest{-AutomatedPerfTest.DoCSVProfiler -AutomatedPerfTest.DoGPUReshape $(GPUReshapeTestCommonArgs)}[Material GPUReshape]" If="$(DoGPUReshape)"/>
-		<!-- Zen Tests -->
-		<Property Name="ZenTestArgs" Value="-AutomatedPerfTest.DoBuildZenStreaming"/>
 		<Property Name="ZenStreamingTargetTestList" Value="$(ZenStreamingTargetTestList)+AutomatedPerfTest.MaterialTest{-AutomatedPerfTest.DoCSVProfiler $(PerfTestCommonArgs) $(ZenTestArgs)}[Zen Material Perf]" If="$(DoPerf) and $(DoBuildZenStreaming)" />
 		<Property Name="ZenStreamingTargetTestList" Value="$(ZenStreamingTargetTestList)+AutomatedPerfTest.MaterialTest{-AutomatedPerfTest.DoCSVProfiler -AutomatedPerfTest.DoLLM $(LLMTestCommonArgs) $(ZenTestArgs)}[Zen Material LLM]" If="$(DoLLM) and $(DoBuildZenStreaming)"/>
 		<Property Name="ZenStreamingTargetTestList" Value="$(ZenStreamingTargetTestList)+AutomatedPerfTest.MaterialTest{-AutomatedPerfTest.DoCSVProfiler -AutomatedPerfTest.DoGPUPerf $(LLMTestCommonArgs) $(ZenTestArgs)}[Zen Material GPUPerf]" If="$(DoGPUPerf) and $(DoBuildZenStreaming)"/>
@@ -174,7 +169,6 @@
 		<Property Name="TargetTestList" Value="$(TargetTestList)+AutomatedPerfTest.ReplayTest{-AutomatedPerfTest.ReplayPerfTest.ReplayName=&quot;$(ReplayName)&quot; -AutomatedPerfTest.DoCSVProfiler -AutomatedPerfTest.DoInsightsTrace $(InsightsTestCommonArgs)}[Replay Insights]" If="$(DoInsightsTrace)"/>
 		<Property Name="TargetTestList" Value="$(TargetTestList)+AutomatedPerfTest.ReplayTest{-AutomatedPerfTest.ReplayPerfTest.ReplayName=&quot;$(ReplayName)&quot; -AutomatedPerfTest.DoGPUReshape $(GPUReshapeTestCommonArgs)}[Replay GPUReshape]" If="$(DoGPUReshape)"/>
 		<!-- Zen Tests -->
-		<Property Name="ZenTestArgs" Value="-AutomatedPerfTest.DoBuildZenStreaming"/>
 		<Property Name="ZenStreamingTargetTestList" Value="$(ZenStreamingTargetTestList)+AutomatedPerfTest.ReplayTest{-AutomatedPerfTest.ReplayPerfTest.ReplayName=&quot;$(ReplayName)&quot; -AutomatedPerfTest.DoCSVProfiler $(PerfTestCommonArgs) $(ZenTestArgs)}[Zen Replay Perf]" If="$(DoPerf) and $(DoBuildZenStreaming)" />
 		<Property Name="ZenStreamingTargetTestList" Value="$(ZenStreamingTargetTestList)+AutomatedPerfTest.ReplayTest{-AutomatedPerfTest.ReplayPerfTest.ReplayName=&quot;$(ReplayName)&quot; -AutomatedPerfTest.DoCSVProfiler -AutomatedPerfTest.DoLLM $(LLMTestCommonArgs) $(ZenTestArgs)}[Zen Replay LLM]" If="$(DoLLM) and $(DoBuildZenStreaming)"/>
 		<Property Name="ZenStreamingTargetTestList" Value="$(ZenStreamingTargetTestList)+AutomatedPerfTest.ReplayTest{-AutomatedPerfTest.ReplayPerfTest.ReplayName=&quot;$(ReplayName)&quot; -AutomatedPerfTest.DoCSVProfiler -AutomatedPerfTest.DoGPUPerf $(LLMTestCommonArgs) $(ZenTestArgs)}[Zen Replay GPUPerf]" If="$(DoGPUPerf) and $(DoBuildZenStreaming)"/>
diff --git a/Engine/Plugins/Performance/AutomatedPerfTesting/Build/Scripts/AutomatedPerfTestConfig.cs b/Engine/Plugins/Performance/AutomatedPerfTesting/Build/Scripts/AutomatedPerfTestConfig.cs
--- a/Engine/Plugins/Performance/AutomatedPerfTesting/Build/Scripts/AutomatedPerfTestConfig.cs
+++ b/Engine/Plugins/Performance/AutomatedPerfTesting/Build/Scripts/AutomatedPerfTestConfig.cs
@@ -117,11 +117,23 @@ namespace AutomatedPerfTest
 		public string GPUReshapeWorkspace;
 
 		/// <summary>
-		/// Whether to run Zen Streaming Import tests with -dobuildzenstreaming
+		/// Whether to run Zen Streaming Import tests with -DoBuildZenStreaming
 		/// </summary>
 		[AutoParamWithNames(false, "AutomatedPerfTest.DoBuildZenStreaming")]
 		public bool DoBuildZenStreaming;
 
+		/// <summary>
+		/// The IP/Web Address of the Zen Server
+		/// </summary>
+		[AutoParamWithNames("", "AutomatedPerfTest.ZenStoreHost")]
+		public string ZenStoreHost;
+
+		/// <summary>
+		/// Whether or not to run a test against a local Zen Server
+		/// </summary>
+		[AutoParamWithNames(false, "AutomatedPerfTest.DoLocalZenTest")]
+		public bool DoLocalZenTest;
+
 		/// <summary>
 		/// Whether to lock dynamic resolution. This is applied automatically when GPUPerf test is run.
 		/// </summary>
diff --git a/Engine/Plugins/Performance/AutomatedPerfTesting/Build/Scripts/AutomatedPerfTestNode.cs b/Engine/Plugins/Performance/AutomatedPerfTesting/Build/Scripts/AutomatedPerfTestNode.cs
--- a/Engine/Plugins/Performance/AutomatedPerfTesting/Build/Scripts/AutomatedPerfTestNode.cs
+++ b/Engine/Plugins/Performance/AutomatedPerfTesting/Build/Scripts/AutomatedPerfTestNode.cs
@@ -906,6 +906,16 @@ namespace AutomatedPerfTest
 			if (Config.DoBuildZenStreaming)
 			{
 				Config.TestID = "APTZen";
+				string ZenHost = Config.ZenStoreHost;
+
+				if(Config.DoLocalZenTest)
+                {
+                    ZenHost = UnrealHelpers.GetHostIpAddress();
+                }
+
+				ClientRole.CommandLineParams.Add("DoBuildZenStreaming");
+				ClientRole.CommandLineParams.AddOrAppendParamValue("ZenStoreHost", ZenHost);
+
 			}
 			ClientRole.CommandLineParams.AddOrAppendParamValue("AutomatedPerfTest.TestID", Config.TestID);
 
@@ -930,11 +940,6 @@ namespace AutomatedPerfTest
 				ClientRole.CommandLineParams.Add("llmcsv");
 			}
 
-			if (Config.DoBuildZenStreaming)
-			{
-				ClientRole.CommandLineParams.Add("DoBuildZenStreaming");
-			}
-
 			if (Config.WaitForAttach)
 			{
 				ClientRole.CommandLineParams.Add("waitforattach");
