Zen APT BuildAndTest Fix

#Synced-CL 48573769

diff --git a/Engine/Build/Graph/Tasks/BuildAndTestProject.xml b/Engine/Build/Graph/Tasks/BuildAndTestProject.xml
--- a/Engine/Build/Graph/Tasks/BuildAndTestProject.xml
+++ b/Engine/Build/Graph/Tasks/BuildAndTestProject.xml
@@ -1024,7 +1024,7 @@
 
 		<!-- Zen Streaming Import, Stage, Publish -->
 		<Agent Name="Stage and Publish $(PlatformAgentName) Zen" Type="$(CompileAgentType)">
-			<Node Name="Zen Stage $(ProjectName) $(TargetPlatform)" Requires="$(PlatformCompileNodeName);$(PlatformSnapshotPublishNodeName)">
+			<Node Name="Zen Stage $(ProjectName) $(TargetPlatform)" Requires="$(PlatformCompileNodeName);$(PlatformSnapshotPublishNodeName)" If="$(DoBuildZenStreaming)">
 				<Property Name="SnapshotFilenamePrefix" Value="SnapshotDescriptor"/>
 				<Property Name="SnapshotFilenamePrefix" Value="$(UE_HORDE_TEMPLATEID)" If="'$(UE_HORDE_TEMPLATEID)' != ''"/>
 				<Property Name="SnapshotIdentifier" Value="$(ProjectName).$(UE_HORDE_JOBID)"/>			
@@ -1051,7 +1051,7 @@
 			</Node>
  
 			<!-- Zen Streaming Publish -->
-			<Node Name="Zen Publish $(ProjectName) $(TargetPlatform) files" Requires="Zen Stage $(ProjectName) $(TargetPlatform)" >
+			<Node Name="Zen Publish $(ProjectName) $(TargetPlatform) files" Requires="Zen Stage $(ProjectName) $(TargetPlatform)" If="$(DoBuildZenStreaming)">
 				<Property Name="TargetDirectory" Value="$(ProjectOutputDirectory)/$(StagedPlatformFolder)"/>
 				<Expand Name="WriteArtifactMetadataFile" Type="staged-build" Platform="$(StagedPlatformFolder)" Configuration="$(PackageConfigurations)" ArtifactDir="$(TargetDirectory)" />
 				<Tag Files="$(TargetDirectory)/..." With="#Published Zen $(ProjectName) $(TargetPlatform) Files"/>
@@ -1290,7 +1290,7 @@
 						<!-- Retrieve Staged Build-->
 						<Property Name="ZenTestPrerequisites" Value="Zen Publish $(ProjectName) $(TargetPlatform) files" />
 						<Property Name="ZenCloudArtifactPlatformRetrieveWithConfig" Value="Zen $(CloudArtifactPlatformRetrieve) $(TestConfiguration)$(TestConfigTitle)"/>
-						<Node Name="$(ZenCloudArtifactPlatformRetrieveWithConfig)" Requires="$(ZenTestPrerequisites)">
+						<Node Name="$(ZenCloudArtifactPlatformRetrieveWithConfig)" Requires="$(ZenTestPrerequisites)" If="$(DoBuildZenStreaming)">
 							<RetrieveCloudArtifact 
 								Name="$(ProjectName)-$(TargetPlatform)-Zen-Staged" 
 								Type="staged-build" 
@@ -1316,7 +1316,7 @@
 								<Property Name="GauntletReportArgs" Value="-logdir=&quot;$(OutputReportPath)&quot; -artifactname=Gauntlet" />
 							</Do>
 
-							<Node Name="$(ZenInstallNodeName)" Requires="$(ZenTestPrerequisites)" Annotations="DeviceReserve=Begin;$(NodeAnnotations)">
+							<Node Name="$(ZenInstallNodeName)" Requires="$(ZenTestPrerequisites)" Annotations="DeviceReserve=Begin;$(NodeAnnotations)" If="$(DoBuildZenStreaming)">
 								<Command Name="RunUnreal" Arguments="-project=$(TargetProject) -test=&quot;UE.InstallOnly&quot; -platform=$(TargetPlatform) -configuration=$(TestConfiguration) -target=$(TargetName) $(GauntletBuildArg) $(GauntletTestArgs) $(GauntletReportArgs) $(InstallArgs)" />
 							</Node>
 
@@ -1356,7 +1356,7 @@
 							<Property Name="GauntletTestArgs" Value="$(GauntletTestArgs) -adbkeys=&quot;$(AdbKeysPath)&quot;" If="'$(AdbKeysPath)' != '' and '$(TargetPlatform)' == 'Android'"/>
 
 							<!-- 'Requires' should be the package game and the publish folder is used; however some tests may require the editor ie: UE.TargetAutomation -->
-							<Node Name="$(TestNodeName)" Requires="$(PlatformToolsNodeName);$(ZenTestPrerequisites)" Annotations="$(NodeAnnotations)">
+							<Node Name="$(TestNodeName)" Requires="$(PlatformToolsNodeName);$(ZenTestPrerequisites)" Annotations="$(NodeAnnotations)" If="$(DoBuildZenStreaming)">
 								<Property Name="ZenHost" Value="$(ZenStoreHost)"/>
 								<Property Name="ZenHost" Value="127.0.0.1" If="$(DoLocalZenTest)"/>
 								<Property Name="ZenPlatform" Value="$(TargetPlatform)"/>
@@ -1402,7 +1402,7 @@
 		<Property Name="PlatformTasks" Value="$(PlatformTasks);$(PlatformHashNodeName)"  If="$(PublishHashedStagingData)" />	
 		<Property Name="PlatformTasks" Value="$(PlatformTasks);$(AllPlatformBootTestNodes);$(AllPlatformTestNodes)" />
 		<Property Name="PlatformTasks" Value="$(PlatformTasks);$(PlatformCheckBuildSizeNodeName)" If="'$(BuildSizeLimit)' != ''" />
-
+		
 		<Property Name="GeneratedTargetNodes" Value="$(GeneratedTargetNodes);$(PlatformTasks)"/>
 		<Property Name="GeneratedTargetTestNodes" Value="$(GeneratedTargetTestNodes);$(AllPlatformBootTestNodes);$(AllPlatformTestNodes)"/>
 	</ForEach>
