Fix a build error 'Reference to 'Zen Development Install Game' cannot be resolved'

BuildAndTestProject.xml. was conditionally creating a node $(ZenInstallNodeName) when $(DoBuildZenStreaming) was true but was unconditionally appending it to $(AllPlatformInstallNodes). This handling now matches the handling used when appending to $(AllPlatformTestNodes).

#Synced-CL 45978330

diff --git a/Engine/Build/Graph/Tasks/BuildAndTestProject.xml b/Engine/Build/Graph/Tasks/BuildAndTestProject.xml
--- a/Engine/Build/Graph/Tasks/BuildAndTestProject.xml
+++ b/Engine/Build/Graph/Tasks/BuildAndTestProject.xml
@@ -1323,7 +1323,7 @@
 							<Property Name="ZenTestPrerequisites" Value="$(ZenTestPrerequisites);$(ZenInstallNodeName)" />
 
 							<Do If="!ContainsItem('$(TargetPlatformsToSkipBootTests)', '$(TargetPlatform)', '+') Or !ContainsItem('$(TargetPlatformsToSkipTests)', '$(TargetPlatform)', '+')">
-								<Property Name="AllPlatformInstallNodes" Value="$(AllPlatformInstallNodes);$(ZenInstallNodeName)"/>
+								<Property Name="AllPlatformInstallNodes" Value="$(AllPlatformInstallNodes);$(ZenInstallNodeName)" If="$(DoBuildZenStreaming)"/>
 							</Do>
 
 							<!-- Skip install for all other test node -->
