Introduce TestDesktopPlatform*AgentOverride and Extra*CompileArguments options

#rnx

#Synced-CL 48759964

diff --git a/Engine/Build/Graph/Tasks/BuildAndTestProject.xml b/Engine/Build/Graph/Tasks/BuildAndTestProject.xml
--- a/Engine/Build/Graph/Tasks/BuildAndTestProject.xml
+++ b/Engine/Build/Graph/Tasks/BuildAndTestProject.xml
@@ -43,6 +43,9 @@
 	<Property Name="PreLabelName" Value="$(PreLabelName) " If="'$(PreLabelName)' != ''"/>
 
 	<Option Name="GenericCompileArguments" DefaultValue="" Restrict="" Description="Configurable compile args applied to all compile targets."/>
+	<Option Name="ExtraWin64CompileArguments" DefaultValue="" Restrict="" Description="Extra arguments to use when building Win64 compile targets."/>
+	<Option Name="ExtraLinuxCompileArguments" DefaultValue="" Restrict="" Description="Extra arguments to use when building Linux compile targets."/>
+	<Option Name="ExtraMacCompileArguments" DefaultValue="" Restrict="" Description="Extra arguments to use when building Mac compile targets."/>
 	<Option Name="ExtraToolCompileArguments" DefaultValue="" Restrict="" Description="Extra arguments to use when building the tools" />
 	<Option Name="ExtraEditorCompileArguments" DefaultValue="" Restrict="" Description="Extra arguments to use when building the editor" />
 	<Option Name="ExtraTargetCompileArguments" DefaultValue="" Restrict="" Description="Extra arguments to use when building targets" />
@@ -60,6 +63,9 @@
 	<Option Name="AgentOverride" DefaultValue="" Description="If set, other logic is ignored and this agent type is used for all non-test work." />
 	<Option Name="TestAgentOverride" DefaultValue="" Description="If set, other logic is ignored and this agent type is used for all test work." />
 	<Option Name="TargetPlatformTestAgentOverride" DefaultValue="DevkitAutomation" Description="Override agent type use to run target platform tests"/>
+	<Option Name="TestDesktopPlatformWin64AgentOverride" DefaultValue="" Description="Override agent type use to run tests for Win64 target"/>
+	<Option Name="TestDesktopPlatformLinuxAgentOverride" DefaultValue="" Description="Override agent type use to run tests for Linux target"/>
+	<Option Name="TestDesktopPlatformMacAgentOverride" DefaultValue="" Description="Override agent type use to run tests for Mac target"/>
 	<Option Name="ClientCompileAgentOverride" DefaultValue="" Description="Specifies the agent type to use for C++ compilation jobs for target platform binaries." />
 	<Option Name="DontUseTestAgent" Restrict="true|false" DefaultValue="false" Description="Avoid using test agent pool to run tests"/>
 	<Option Name="UnrealInsightsBuildConfiguration" DefaultValue="Shipping" Restrict="Development|Debug|Shipping" Description="In which configuration to build Unreal Insights, if included. Shipping, Debug, or Development." />
@@ -321,6 +327,9 @@
 		<!-- Use overrides for agent types? -->
 		<Property Name="HostAgentType" Value="$(AgentOverride)" If="'$(AgentOverride)' != ''" />
 		<Property Name="TestAgentType" Value="$(TestAgentOverride)"  If="'$(TestAgentOverride)' != ''" />
+		<Property Name="TestAgentType" Value="$(TestDesktopPlatformWin64AgentOverride)" If="'$(TestDesktopPlatformWin64AgentOverride)' != '' and '$(EditorPlatform)' == 'Win64'"/>
+		<Property Name="TestAgentType" Value="$(TestDesktopPlatformLinuxAgentOverride)" If="'$(TestDesktopPlatformLinuxAgentOverride)' != '' and '$(EditorPlatform)' == 'Linux'"/>
+		<Property Name="TestAgentType" Value="$(TestDesktopPlatformMacAgentOverride)" If="'$(TestDesktopPlatformMacAgentOverride)' != '' and '$(EditorPlatform)' == 'Mac'"/>
 
 		<!-- Test nodes for this editor platform -->
 		<Property Name="PlatformEditorTestNodes" Value="" />
@@ -352,23 +361,27 @@
 				<Compile Target="UnrealInsights" Platform="$(EditorPlatform)" Configuration="$(UnrealInsightsBuildConfiguration)" Arguments="$(GenericCompileArguments)" If="$(IncludeUnrealInsightsInToolCompile)"/>
 				<Compile Target="EpicWebHelper" Platform="$(EditorPlatform)" Configuration="Shipping" Arguments="$(GenericCompileArguments)" If="$(IncludeEpicWebHelperInToolCompile)"/>
 			</Node>
+			<Property Name="CompileArguments" Value="$(GenericCompileArguments)"/>
+			<Property Name="CompileArguments" Value="$(CompileArguments) $(ExtraWin64CompileArguments)" If="'$(EditorPlatform)' == 'Win64' and '$(ExtraWin64CompileArguments)' != ''"/>
+			<Property Name="CompileArguments" Value="$(CompileArguments) $(ExtraLinuxCompileArguments)" If="'$(EditorPlatform)' == 'Linux' and '$(ExtraLinuxCompileArguments)' != ''"/>
+			<Property Name="CompileArguments" Value="$(CompileArguments) $(ExtraMacCompileArguments)" If="'$(EditorPlatform)' == 'Mac' and '$(ExtraMacCompileArguments)' != ''"/>
 			<Property Name="CompileEditor" Value="$(PreNodeName)Compile Editor $(EditorPlatform)"/>
 			<Node Name="$(CompileEditor)" Requires="$(PreNodeName)Compile Tools $(EditorPlatform)">
 				<Switch>
 					<Case If="$(IsContentOnlyProject)">
-						<Command Name="BuildCookRun" Arguments="-project=&quot;$(RootDir)/$(TargetProject)&quot; -platform=$(EditorPlatform) -configuration=Development -build -notools -ubtargs=&quot;$(GenericCompileArguments) $(ExtraEditorCompileArguments)&quot;"/>
+						<Command Name="BuildCookRun" Arguments="-project=&quot;$(RootDir)/$(TargetProject)&quot; -platform=$(EditorPlatform) -configuration=Development -build -notools -ubtargs=&quot;$(CompileArguments) $(ExtraEditorCompileArguments)&quot;"/>
 						<Tag Files="Engine/Binaries/...;$(ProjectPath)/Binaries/..." Filter="UnrealEditor.target;$(EditorTargetName).target" With="#$(CompileEditor)"/>
 						<TagReceipt Files="#$(CompileEditor)" BuildProducts="true" With="#$(CompileEditor)"/>
 					</Case>
 					<Default>
-						<Compile Target="$(EditorTargetName)" Project="$(TargetProject)" Platform="$(EditorPlatform)" Configuration="Development" Arguments="$(GenericCompileArguments) $(ExtraEditorCompileArguments)"/>
+						<Compile Target="$(EditorTargetName)" Project="$(TargetProject)" Platform="$(EditorPlatform)" Configuration="Development" Arguments="$(CompileArguments) $(ExtraEditorCompileArguments)"/>
 					</Default>
 				</Switch>
 			</Node>
 			<!-- Cooked editor target -->
 			<Node Name="$(PreNodeName)Compile CookedEditor $(EditorPlatform)" Requires="$(PreNodeName)Compile Tools $(EditorPlatform)">
 				<ForEach Name="Config" Values="$(CookedEditorConfigurations)" Separator="+">
-					<Compile Target="$(CookedEditorTargetName)" Platform="$(EditorPlatform)" Configuration="$(Config)" Arguments="$(GenericCompileArguments) $(ExtraEditorCompileArguments)"/>
+					<Compile Target="$(CookedEditorTargetName)" Platform="$(EditorPlatform)" Configuration="$(Config)" Arguments="$(CompileArguments) $(ExtraEditorCompileArguments)"/>
 				</ForEach>
 			</Node>
 			<!--Add these nodes to our dependency list -->
@@ -452,6 +465,9 @@
 		<Property Name="TestAgentType" Value="Test$(EditorPlatform)" />
 		<Property Name="TestAgentType" Value="$(EditorPlatform)" If="$(DontUseTestAgent)" />
 		<Property Name="TestAgentType" Value="$(TestAgentOverride)" If="'$(TestAgentOverride)' != ''"/>
+		<Property Name="TestAgentType" Value="$(TestDesktopPlatformWin64AgentOverride)" If="'$(TestDesktopPlatformWin64AgentOverride)' != '' and '$(EditorPlatform)' == 'Win64'"/>
+		<Property Name="TestAgentType" Value="$(TestDesktopPlatformLinuxAgentOverride)" If="'$(TestDesktopPlatformLinuxAgentOverride)' != '' and '$(EditorPlatform)' == 'Linux'"/>
+		<Property Name="TestAgentType" Value="$(TestDesktopPlatformMacAgentOverride)" If="'$(TestDesktopPlatformMacAgentOverride)' != '' and '$(EditorPlatform)' == 'Mac'"/>
 
 		<Property Name="EditorPlatformPath" Value="$(EditorPlatform)" />
 		<Property Name="EditorPlatformPath" Value="Windows" If="'$(EditorPlatform)' == 'Win64'" />
@@ -622,6 +638,9 @@
 		<Property Name="TestAgentType" Value="$(TargetPlatformTestAgentOverride)" />
 		<Property Name="TestAgentType" Value="Test$(HostAgentType)" If="$(IsDesktopPlatform)"/>
 		<Property Name="TestAgentType" Value="$(HostAgentType)" If="$(IsDesktopPlatform) and $(DontUseTestAgent)" />
+		<Property Name="TestAgentType" Value="$(TestDesktopPlatformWin64AgentOverride)" If="$(IsDesktopPlatform) and '$(TestDesktopPlatformWin64AgentOverride)' != '' and $(IsTargetPlatformWindows)"/>
+		<Property Name="TestAgentType" Value="$(TestDesktopPlatformLinuxAgentOverride)" If="$(IsDesktopPlatform) and '$(TestDesktopPlatformLinuxAgentOverride)' != '' and $(IsTargetPlatformLinux)"/>
+		<Property Name="TestAgentType" Value="$(TestDesktopPlatformMacAgentOverride)" If="$(IsDesktopPlatform) and '$(TestDesktopPlatformMacAgentOverride)' != '' and $(IsTargetPlatformMac)"/>
 
 		<!-- Use incremental agents for building? -->
 		<Property Name="CompileAgentType" Value="$(HostAgentType)" />
@@ -693,6 +712,11 @@
 		<Property Name="BCRArgs" Value="-project=$(TargetProject) -platform=$(TargetPlatform) -cookflavor=&quot;$(BCRCookFlavor)&quot;" />
 		<Property Name="BCRArgs" Value="$(BCRArgs) -client" If="'$(TargetType)' == 'Client'" />
 
+		<Property Name="CompileArguments" Value="$(GenericCompileArguments)"/>
+		<Property Name="CompileArguments" Value="$(CompileArguments) $(ExtraWin64CompileArguments)" If="$(IsDesktopPlatform) and '$(ExtraWin64CompileArguments)' != '' and $(IsTargetPlatformWindows)"/>
+		<Property Name="CompileArguments" Value="$(CompileArguments) $(ExtraLinuxCompileArguments)" If="$(IsDesktopPlatform) and '$(ExtraLinuxCompileArguments)' != '' and $(IsTargetPlatformLinux)"/>
+		<Property Name="CompileArguments" Value="$(CompileArguments) $(ExtraMacCompileArguments)" If="$(IsDesktopPlatform) and '$(ExtraMacCompileArguments)' != '' and $(IsTargetPlatformMac)"/>
+
 		<!-- Support additional stage and package arguments that are only active on some platforms -->
 		<Property Name="StageAndPackageArguments" Value="$(ExtraStageAndPackageArguments)" />
 		<Property Name="StageAndPackageArguments" Value="$(StageAndPackageArguments) $(TargetPlatformFilteredExtraStageAndPackageArguments)" If="'$(TargetPlatformFilteredExtraStageAndPackageArgumentsFilter)' != '' and ContainsItem('$(TargetPlatformFilteredExtraStageAndPackageArgumentsFilter)', '$(TargetPlatform)', '+')" />
@@ -710,7 +734,7 @@
 			<Node Name="$(PlatformCompileNodeName)" Requires="$(PlatformToolsNodeName)">
 				<Switch>
 					<Case If="$(IsContentOnlyProject)">
-						<Command Name="BuildGame" Arguments="$(BCRArgs) -configuration=$(TargetConfigurations) -notools -ubtargs=&quot;-BuildVersion=$(BuildVersion) $(GenericCompileArguments) $(TargetCompileArguments)&quot;" />
+						<Command Name="BuildGame" Arguments="$(BCRArgs) -configuration=$(TargetConfigurations) -notools -ubtargs=&quot;-BuildVersion=$(BuildVersion) $(CompileArguments) $(TargetCompileArguments)&quot;" />
 						<Tag Files="$(ProjectPath)/Binaries/..." Filter="$(TargetName)*.target" With="#$(PlatformCompileNodeName)"/>
 						<Tag Files="$(ProjectPath)/Binaries/..." Filter="UnrealGame*.target" With="#UnrealGameTarget"/>
 						<Copy Files="#UnrealGameTarget" From="$(ProjectPath)" To="Engine" Tag="#$(PlatformCompileNodeName)"/>
@@ -718,7 +742,7 @@
 					</Case>
 					<Default>
 						<ForEach Name="TargetConfiguration" Values="$(TargetConfigurations)" Separator="+">
-							<Compile Target="$(TargetName)" Project="$(TargetProject)" Platform="$(TargetPlatform)" Configuration="$(TargetConfiguration)" Arguments="-BuildVersion=&quot;$(BuildVersion)&quot; $(GenericCompileArguments) $(TargetCompileArguments)"/>
+							<Compile Target="$(TargetName)" Project="$(TargetProject)" Platform="$(TargetPlatform)" Configuration="$(TargetConfiguration)" Arguments="-BuildVersion=&quot;$(BuildVersion)&quot; $(CompileArguments) $(TargetCompileArguments)"/>
 						</ForEach>
 					</Default>
 				</Switch>
@@ -728,7 +752,7 @@
 				<ForEach Name="TargetConfiguration" Values="$(TargetConfigurations)" Separator="+">
 					<Property Name="ServerPlatform" Value="Win64"/>
 					<Property Name="ServerPlatform" Value="Linux" If="$(IsTargetPlatformLinux)"/>
-					<Property Name="ServerCompileArguments" Value="-BuildVersion=&quot;$(BuildVersion)&quot; $(GenericCompileArguments) $(TargetCompileArguments)"/>
+					<Property Name="ServerCompileArguments" Value="-BuildVersion=&quot;$(BuildVersion)&quot; $(CompileArguments) $(TargetCompileArguments)"/>
 					<Compile Target="$(TargetNameServer)" Platform="$(ServerPlatform)" Configuration="$(TargetConfiguration)" Arguments="$(ServerCompileArguments)"/>
 				</ForEach>	
 			</Node>
