Fix missing Editor platform compilation when skipping some platform tests

#rnx

#Synced-CL 48781209

diff --git a/Engine/Build/Graph/Tasks/BuildAndTestProject.xml b/Engine/Build/Graph/Tasks/BuildAndTestProject.xml
--- a/Engine/Build/Graph/Tasks/BuildAndTestProject.xml
+++ b/Engine/Build/Graph/Tasks/BuildAndTestProject.xml
@@ -332,6 +332,7 @@
 		<Property Name="TestAgentType" Value="$(TestDesktopPlatformMacAgentOverride)" If="'$(TestDesktopPlatformMacAgentOverride)' != '' and '$(EditorPlatform)' == 'Mac'"/>
 
 		<!-- Test nodes for this editor platform -->
+		<Property Name="PlatformToolNodes" Value="" />
 		<Property Name="PlatformEditorTestNodes" Value="" />
 
 		<Agent Name="Build $(ProjectName)Editor $(EditorPlatform)" Type="$(HostAgentType)">
@@ -385,14 +386,15 @@
 				</ForEach>
 			</Node>
 			<!--Add these nodes to our dependency list -->
-			<Property Name="GeneratedToolNodes" Value="$(CompileEditor)" If="!$(SkipEditorBuild)" />
-			<Property Name="GeneratedToolNodes" Value="$(GeneratedToolNodes);$(PreNodeName)Compile Tools $(EditorPlatform)" If="!$(SkipBuild)" />
+			<Property Name="PlatformToolNodes" Value="$(CompileEditor)" If="!$(SkipEditorBuild)" />
+			<Property Name="PlatformToolNodes" Value="$(PlatformToolNodes);$(PreNodeName)Compile Tools $(EditorPlatform)" If="!$(SkipBuild)" />
+			<Property Name="GeneratedToolNodes" Value="$(GeneratedToolNodes);$(PlatformToolNodes)" />
 			<!-- Editor Labels -->
-			<Label Category="Editors" Name="$(PreLabelName)Editor $(EditorPlatform)" Requires="$(GeneratedToolNodes)" />
+			<Label Category="Editors" Name="$(PreLabelName)Editor $(EditorPlatform)" Requires="$(PlatformToolNodes)" />
 		</Agent>
 		<Agent Name="Test $(ProjectName)Editor $(EditorPlatform)" Type="$(TestAgentType)">
 			<!-- Prerequisites for the editor test to run -->
-			<Property Name="TestPrerequisites" Value="$(GeneratedToolNodes)"/>
+			<Property Name="TestPrerequisites" Value="$(PlatformToolNodes)"/>
 			<!-- Declare a boot test for the editor-->
 			<Property Name="PlatformBootTestNodeName" Value="$(PreNodeName)Editor $(EditorPlatform) Test=BootTest" />
 			<Do If="'$(AutomationReportOutputDirectory)' != ''">
@@ -408,7 +410,7 @@
 				<Command Name="RunUnreal" Arguments="-project=$(TargetProject) -build=Editor -test=&quot;UE.EditorBootTest&quot; $(GauntletEditorTestArgs) $(GauntletReportArgs)" />
 			</Node>
 			<Property Name="ProjectEditorDDC" Value="$(PreNodeName)Editor DDC $(EditorPlatform)" />
-			<Node Name="$(ProjectEditorDDC)" Requires="$(GeneratedToolNodes)">
+			<Node Name="$(ProjectEditorDDC)" Requires="$(PlatformToolNodes)">
 				<Property Name="DDCTargetPlatform" Value="$(EditorPlatform)" />
 				<Property Name="DDCTargetPlatform" Value="Windows" If="'$(DDCTargetPlatform)' == 'Win64'"/>
 				<Commandlet Name="DerivedDataCache" Project="$(RootDir)/$(TargetProject)" Arguments="-fill -targetplatform=$(DDCTargetPlatform)Editor"/>
