Fix incorrect boolean logic for install node collection in the BuildGraph test configuration.

The install node collection was using OR logic when checking skip lists, which allowed install nodes to be added for platforms that should not run tests. Specifically, a platform listed in TargetPlatformsToSkipTests would still have its install node added to AllPlatformInstallNodes as long as it was not also present in SkipBootTests.

This change updates the condition to use AND logic so install nodes are only added when the platform appears in neither skip list. This matches the intended behavior and brings the logic in line with how boot tests and regular tests apply skip handling elsewhere in the file.

#Synced-CL 49896846

diff --git a/Engine/Build/Graph/Tasks/BuildAndTestProject.xml b/Engine/Build/Graph/Tasks/BuildAndTestProject.xml
--- a/Engine/Build/Graph/Tasks/BuildAndTestProject.xml
+++ b/Engine/Build/Graph/Tasks/BuildAndTestProject.xml
@@ -1152,7 +1152,7 @@
 								<Command Name="RunUnreal" Arguments="-project=$(TargetProject) -test=&quot;UE.InstallOnly&quot; -platform=$(TargetPlatform) -configuration=$(TestConfiguration) -target=$(TargetName) $(GauntletBuildArg) $(GauntletTestArgs) $(GauntletReportArgs) $(InstallArgs)" />
 							</Node>
 							<Property Name="TestPrerequisites" Value="$(TestPrerequisites);$(InstallNodeName)" />
-							<Do If="!ContainsItem('$(TargetPlatformsToSkipBootTests)', '$(TargetPlatform)', '+') Or !ContainsItem('$(TargetPlatformsToSkipTests)', '$(TargetPlatform)', '+')">
+							<Do If="!ContainsItem('$(TargetPlatformsToSkipBootTests)', '$(TargetPlatform)', '+') And !ContainsItem('$(TargetPlatformsToSkipTests)', '$(TargetPlatform)', '+')">
 								<Property Name="AllPlatformInstallNodes" Value="$(AllPlatformInstallNodes);$(InstallNodeName)"/>
 							</Do>
 							<!-- Skip install for all other test node -->
diff --git a/Engine/Plugins/Performance/AutomatedPerfTesting/Build/Inc/AutomatedPerfTestZenStreaming.xml b/Engine/Plugins/Performance/AutomatedPerfTesting/Build/Inc/AutomatedPerfTestZenStreaming.xml
--- a/Engine/Plugins/Performance/AutomatedPerfTesting/Build/Inc/AutomatedPerfTestZenStreaming.xml
+++ b/Engine/Plugins/Performance/AutomatedPerfTesting/Build/Inc/AutomatedPerfTestZenStreaming.xml
@@ -252,7 +252,7 @@
 
 								<Property Name="ZenTestPrerequisites" Value="$(ZenTestPrerequisites);$(ZenInstallNodeName)" />
 
-								<Do If="!ContainsItem('$(TargetPlatformsToSkipBootTests)', '$(TargetPlatform)', '+') Or !ContainsItem('$(TargetPlatformsToSkipTests)', '$(TargetPlatform)', '+')">
+								<Do If="!ContainsItem('$(TargetPlatformsToSkipBootTests)', '$(TargetPlatform)', '+') And !ContainsItem('$(TargetPlatformsToSkipTests)', '$(TargetPlatform)', '+')">
 									<Property Name="AllPlatformInstallNodes" Value="$(AllPlatformInstallNodes);$(ZenInstallNodeName)"/>
 								</Do>
 
