Add initial support for variations, with arm64, arm64ec and multi-architecture x64+arm64 available for win64:

-set:WithArm64Variation=true
-set:WithArm64ecVariation=true
-set:WithMultiArchVariation=true

#Synced-CL 50434826

#Synced-CL 50434826

diff --git a/Engine/Build/Graph/Tasks/BuildAndTestProject.xml b/Engine/Build/Graph/Tasks/BuildAndTestProject.xml
--- a/Engine/Build/Graph/Tasks/BuildAndTestProject.xml
+++ b/Engine/Build/Graph/Tasks/BuildAndTestProject.xml
@@ -605,652 +605,680 @@
 	<!-- Now, for each platform we build, cook, stage, package, publish, and test -->
 	<ForEach Name="TargetPlatform" Values="$(TargetPlatforms)" Separator="+">
 
-		<Property Name="IsDesktopPlatform" Value="ContainsItem('$(AllDesktopPlatforms)', '$(TargetPlatform)', '+')"/>
-		<Property Name="IsMobilePlatform" Value="('$(TargetPlatform)' == 'Android' or '$(TargetPlatform)' == 'IOS')"/>
-		<!-- Local or virtualized platform -->
-		<Property Name="IsVirtualizedPlatform" Value="ContainsItem('$(AllVirtualizedPlatforms)', '$(TargetPlatform)', '+')"/>
-		<Property Name="IsLocalPlatform" Value="($(IsDesktopPlatform) Or $(IsVirtualizedPlatform))"/>
-
-		<Property Name="IsTargetPlatformMac" Value="('$(TargetPlatform)' == 'Mac' or '$(TargetPlatform)' == 'IOS' or '$(TargetPlatform)' == 'tvOS')"/>
-		<Property Name="IsTargetPlatformLinux" Value="('$(TargetPlatform)' == 'Linux')"/>
-		<Property Name="IsTargetPlatformWindows" Value="(!$(IsTargetPlatformMac) and !$(IsTargetPlatformLinux))"/>
-		
-		<!-- Declare type of host required to build & cook this platform-->
-		<Property Name="HostAgentType" Value="$(EditorPlatformWin64AgentOverride)" /> 
-		<Property Name="HostAgentType" Value="$(EditorPlatformMacAgentOverride)" If="$(IsTargetPlatformMac)"/>
-		<Property Name="HostAgentType" Value="$(EditorPlatformLinuxAgentOverride)" If="$(IsTargetPlatformLinux)"/>
-
-		<Property Name="IsEditorAgentOverride" Value="false"/>
-		<Property Name="IsEditorAgentOverride" Value="true" If="$(IsTargetPlatformWindows) and '$(EditorPlatformWin64AgentOverride)'!='Win64'"/>
-		<Property Name="IsEditorAgentOverride" Value="true" If="$(IsTargetPlatformMac) and '$(EditorPlatformMacAgentOverride)'!='Mac'"/>
-		<Property Name="IsEditorAgentOverride" Value="true" If="$(IsTargetPlatformLinux) and '$(EditorPlatformLinuxAgentOverride)'!='Linux'"/>
-		<!-- On consoles we can use agents with devkit setup, but locally hosted tests need a Test agent-->
-		<Property Name="TestAgentType" Value="$(TargetPlatformTestAgentOverride)" />
-		<Property Name="TestAgentType" Value="Test$(HostAgentType)" If="$(IsDesktopPlatform)"/>
-		<Property Name="TestAgentType" Value="$(HostAgentType)" If="$(IsDesktopPlatform) and $(DontUseTestAgent)" />
-		<Property Name="TestAgentType" Value="$(TestDesktopPlatformWin64AgentOverride)" If="$(IsDesktopPlatform) and '$(TestDesktopPlatformWin64AgentOverride)' != '' and $(IsTargetPlatformWindows)"/>
-		<Property Name="TestAgentType" Value="$(TestDesktopPlatformLinuxAgentOverride)" If="$(IsDesktopPlatform) and '$(TestDesktopPlatformLinuxAgentOverride)' != '' and $(IsTargetPlatformLinux)"/>
-		<Property Name="TestAgentType" Value="$(TestDesktopPlatformMacAgentOverride)" If="$(IsDesktopPlatform) and '$(TestDesktopPlatformMacAgentOverride)' != '' and $(IsTargetPlatformMac)"/>
-
-		<!-- Use incremental agents for building? -->
-		<Property Name="CompileAgentType" Value="$(HostAgentType)" />
-		<Property Name="CompileAgentType" Value="Incremental$(HostAgentType)" If="$(UseIncrementalAgents)" />
-
-		<!-- Use agent types from override? -->
-		<Property Name="CompileAgentType" Value="$(AgentOverride)" If="'$(AgentOverride)' != ''"/>
-		<!-- Use the same agent type for packaging as compilation -->
-		<Property Name="PackageAgentType" Value="$(CompileAgentType)"/>
-		<Property Name="CompileAgentType" Value="$(ClientCompileAgentOverride)" If="'$(ClientCompileAgentOverride)' != ''"/>
-		<Property Name="TestAgentType" Value="$(TestAgentOverride)" If="'$(TestAgentOverride)' != ''"/>
-
-		<Property Name="PreNodeNameWithPlatform" Value="$(PreNodeName)"/>
-		<Property Name="PreNodeNameWithPlatform" Value="[$(HostAgentType)]$(PreNodeName)" If="$(IsEditorAgentOverride)"/>
-		<!-- Building and cooking will require these nodes from the appropriate host. They should have been built above... -->
-		<Property Name="PlatformToolsNodeName" Value="$(PreNodeName)Compile Tools $(HostAgentType)" />
-		<Property Name="PlatformEditorNodeName" Value="$(PreNodeName)Compile Editor $(HostAgentType)" />
-		<!-- Define these names up front as we use some of them several times -->
-		<Property Name="PlatformAgentName" Value="$(ProjectName) $(TargetPlatform)" />
-		<Property Name="PlatformCompileNodeName" Value="$(PreNodeNameWithPlatform)Compile $(TargetPlatform)" />
-		<Property Name="UploadSymbolsNodeName" Value="" />
-		<Property Name="UploadSymbolsNodeName" Value="$(PreNodeNameWithPlatform)Upload Symbols $(TargetPlatform)" If="'$(SymbolServerPath)' != '' And !$(IsMobilePlatform) And !$(IsTargetPlatformLinux)" />
-		<Property Name="PlatformCookNodeName" Value="$(PreNodeNameWithPlatform)Cook $(TargetPlatform)" />
-		<Property Name="PlatformStageNodeName" Value="$(PreNodeNameWithPlatform)Stage $(TargetPlatform)" />
-		<Property Name="PlatformPackageNodeName" Value="$(PreNodeNameWithPlatform)Package $(TargetPlatform)" />
-		<Property Name="PlatformPackageArtifactName" Value="$(PlatformPackageNodeName)" />
-		<Property Name="PlatformHashNodeName" Value="$(PreNodeNameWithPlatform)Hash $(TargetPlatform)" />
-		<Property Name="PlatformPublishStagedNodeName" Value="$(PreNodeNameWithPlatform)Publish Staged $(TargetPlatform)" />
-		<Property Name="PlatformPublishPackagedNodeName" Value="$(PreNodeNameWithPlatform)Publish Packaged $(TargetPlatform)" />
-		<Property Name="PlatformCheckBuildSizeNodeName" Value="$(PreNodeNameWithPlatform)Check Build Size $(TargetPlatform)" />
-		<Property Name="PlatformStagedArtifactNodeName" Value="$(PlatformPublishStagedNodeName) (Horde)" />
-		<Property Name="CloudPlatformStagedArtifactNodeName" Value="$(PlatformPublishStagedNodeName) (Cloud)" />
-		<Property Name="HordeArtifactPlatformPackagedNodeName" Value="$(PlatformPublishPackagedNodeName) (Horde)" />
-		<Property Name="HordeArtifactPlatformRetrieve" Value="$(PreNodeNameWithPlatform)Retrieve Build Artifact $(TargetPlatform)" />
-		<Property Name="CloudArtifactPlatformPackagedNodeName" Value="$(PlatformPublishPackagedNodeName) (Cloud)" />
-		<Property Name="CloudArtifactPlatformRetrieve" Value="$(PreNodeNameWithPlatform)Retrieve Build Artifact $(TargetPlatform)" />
-		<Property Name="RequirePlatformStaged" Value="$(PlatformStageNodeName)" />
-		<Property Name="RequirePlatformStaged" Value="$(PlatformStageNodeName) Server" If="$(DedicatedServer)"/>
-		<Property Name="RequirePlatformStaged" Value="$(RequirePlatformStaged);$(PlatformStagedArtifactNodeName)" If="!$(SkipPublish) And $(PublishHordeArtifacts)"/>
-		<Property Name="RequirePlatformStaged" Value="$(RequirePlatformStaged);$(PlatformStagedArtifactNodeName) Server" If="!$(SkipPublish) And $(PublishHordeArtifacts) And $(DedicatedServer)"/>
-		<Property Name="RequirePlatformStaged" Value="$(RequirePlatformStaged);$(CloudPlatformStagedArtifactNodeName)" If="!$(SkipPublish) And $(PublishCloudArtifacts)"/>
-		<Property Name="RequirePlatformStaged" Value="$(RequirePlatformStaged);$(PlatformPublishStagedNodeName)" If="!$(SkipPublish) And $(PublishBuilds)"/>
-		<Property Name="RequirePlatformStaged" Value="$(RequirePlatformStaged);$(PlatformPublishStagedNodeName) Server" If="!$(SkipPublish) And $(PublishBuilds) And $(DedicatedServer)"/>
-		<Property Name="RequirePlatformPackaged" Value="$(PlatformPackageNodeName)"/>
-		<Property Name="RequirePlatformPackaged" Value="$(PlatformPackageNodeName) Server" If="$(DedicatedServer)"/>
-		<Property Name="RequirePlatformPackaged" Value="$(RequirePlatformPackaged);$(HordeArtifactPlatformPackagedNodeName)" If="!$(SkipPublish) And $(PublishHordeArtifacts)"/>
-		<Property Name="RequirePlatformPackaged" Value="$(RequirePlatformPackaged);$(HordeArtifactPlatformPackagedNodeName) Server" If="!$(SkipPublish) And $(PublishHordeArtifacts) And $(DedicatedServer)"/>
-		<Property Name="RequirePlatformPackaged" Value="$(RequirePlatformPackaged);$(CloudArtifactPlatformPackagedNodeName)" If="!$(SkipPublish) And $(PublishCloudArtifacts)"/>
-		<Property Name="RequirePlatformPackaged" Value="$(RequirePlatformPackaged);$(PlatformPublishPackagedNodeName)"  If="!$(SkipPublish) And $(PublishBuilds)"/>
-		<Property Name="RequirePlatformPackaged" Value="$(RequirePlatformPackaged);$(PlatformPublishPackagedNodeName) Server"  If="!$(SkipPublish) And $(PublishBuilds) And $(DedicatedServer)"/>
-		
-		<!-- Set CookPlatform since it isn't always the same as the target -->
-		<Property Name="CookPlatform" Value="$(TargetPlatform)" />
-		<Property Name="CookPlatform" Value="Windows" If="'$(TargetPlatform)' == 'Win64'"/>
-		<Property Name="CookPlatform" Value="$(CookPlatform)_$(AndroidCookFlavor)" If="'$(TargetPlatform)' == 'Android'"/>
-		<Property Name="CookPlatform" Value="$(CookPlatform)Client" If="'$(TargetType)' == 'Client'" />
-
-		<!-- Only use Cook Flavor for Android -->
-		<Property Name="BCRCookFlavor" Value=""/>
-        <Property Name="BCRCookFlavor" Value="$(AndroidCookFlavor)" If="'$(TargetPlatform)' == 'Android'" />
-		
-		<Property Name="StagedPlatformFolder" Value="$(CookPlatform)"  />
-		<!-- <Property Name="StagedPlatformFolder" Value="$(StagedPlatformFolder)Client" If="'$(TargetType)' == 'Client'" /> -->
-
-		<!-- Hashing of staged data can only be done on Win64 hosts -->
-		<Property Name="DoUnsyncHashing" Value="$(HostAgentType) == 'Win64' and $(IsEpicInternal)" />
-
-		<!-- Default BuildCookRun args needed when staging/packaging/archiving this configuration -->
-		<Property Name="BCRArgs" Value="-project=$(TargetProject) -platform=$(TargetPlatform) -cookflavor=&quot;$(BCRCookFlavor)&quot;" />
-		<Property Name="BCRArgs" Value="$(BCRArgs) -client" If="'$(TargetType)' == 'Client'" />
-
-		<Property Name="CompileArguments" Value="$(GenericCompileArguments)"/>
-		<Property Name="CompileArguments" Value="$(CompileArguments) $(ExtraWin64CompileArguments)" If="$(IsDesktopPlatform) and '$(ExtraWin64CompileArguments)' != '' and $(IsTargetPlatformWindows)"/>
-		<Property Name="CompileArguments" Value="$(CompileArguments) $(ExtraLinuxCompileArguments)" If="$(IsDesktopPlatform) and '$(ExtraLinuxCompileArguments)' != '' and $(IsTargetPlatformLinux)"/>
-		<Property Name="CompileArguments" Value="$(CompileArguments) $(ExtraMacCompileArguments)" If="$(IsDesktopPlatform) and '$(ExtraMacCompileArguments)' != '' and $(IsTargetPlatformMac)"/>
-
-		<!-- Support additional stage and package arguments that are only active on some platforms -->
-		<Property Name="StageAndPackageArguments" Value="$(ExtraStageAndPackageArguments)" />
-		<Property Name="StageAndPackageArguments" Value="$(StageAndPackageArguments) $(TargetPlatformFilteredExtraStageAndPackageArguments)" If="'$(TargetPlatformFilteredExtraStageAndPackageArgumentsFilter)' != '' and ContainsItem('$(TargetPlatformFilteredExtraStageAndPackageArgumentsFilter)', '$(TargetPlatform)', '+')" />
-		<Property Name="PackageOnlyArguments" Value="$(ExtraPackageOnlyArguments)" />
-
-		<Property Name="StageAndPackageServerArguments" Value="$(ExtraStageAndPackageServerArguments)" />
-
-		<Property Name="BuildSizeLimit" Value=""/>
-		<Regex Pattern="(?&lt;!\w)$(TargetPlatform):([0-9]+)" Capture="BuildSizeLimit" Input="$(CheckBuildSizePlatforms)"/>
-
-		<!-- Build the client and if necessary a server-->
-		<Agent Name="Compile $(PlatformAgentName)" Type="$(CompileAgentType)">
-
-			<!-- Compile -->
-			<Node Name="$(PlatformCompileNodeName)" Requires="$(PlatformToolsNodeName)">
-				<Switch>
-					<Case If="$(IsContentOnlyProject)">
-						<Command Name="BuildGame" Arguments="$(BCRArgs) -configuration=$(TargetConfigurations) -notools -ubtargs=&quot;-BuildVersion=$(BuildVersion) $(CompileArguments) $(TargetCompileArguments)&quot;" />
-						<Tag Files="$(ProjectPath)/Binaries/..." Filter="$(TargetName)*.target" With="#$(PlatformCompileNodeName)"/>
-						<Tag Files="$(ProjectPath)/Binaries/..." Filter="UnrealGame*.target" With="#UnrealGameTarget"/>
-						<Copy Files="#UnrealGameTarget" From="$(ProjectPath)" To="Engine" Tag="#$(PlatformCompileNodeName)"/>
-						<TagReceipt Files="#$(PlatformCompileNodeName)" BuildProducts="true" With="#$(PlatformCompileNodeName)"/>
-					</Case>
-					<Default>
+		<Property Name="Variations" Value="Default"/>
+		<Property Name="Variations" Value="$(Win64Variations)" If="'$(TargetPlatform)' == 'Win64'"/>
+		<Property Name="Variations" Value="$(LinuxVariations)" If="'$(TargetPlatform)' == 'Linux'"/>
+		<Property Name="Variations" Value="$(MacVariations)" If="'$(TargetPlatform)' == 'Mac'"/>
+
+		<ForEach Name="Variation" Values="$(Variations)" Separator="+">
+
+			<Do If="ContainsItem('$(RequestedVariations)', '$(Variation)', '+')">
+
+				<!-- Variation properties -->
+				<Property Name="IsDefaultVariation" Value="('$(Variation)' == 'Default')"/>
+				<Property Name="IsArm64Variation" Value="ContainsItem('$(Variation)', 'arm64', '+')"/>
+				<Property Name="IsArm64ecVariation" Value="ContainsItem('$(Variation)', 'arm64ec', '+')"/>
+				<Property Name="IsMultiArchVariation" Value="ContainsItem('$(Variation)', 'MultiArch', '+')"/>
+
+				<Property Name="VariationTitle" Value=""/>
+				<Property Name="VariationTitle" Value="_$(Variation)" If="!$(IsDefaultVariation)"/>
+
+				<Property Name="IsDesktopPlatform" Value="ContainsItem('$(AllDesktopPlatforms)', '$(TargetPlatform)', '+')"/>
+				<Property Name="IsMobilePlatform" Value="('$(TargetPlatform)' == 'Android' or '$(TargetPlatform)' == 'IOS')"/>
+				<!-- Local or virtualized platform -->
+				<Property Name="IsVirtualizedPlatform" Value="ContainsItem('$(AllVirtualizedPlatforms)', '$(TargetPlatform)', '+')"/>
+				<Property Name="IsLocalPlatform" Value="($(IsDesktopPlatform) Or $(IsVirtualizedPlatform))"/>
+
+				<Property Name="IsTargetPlatformMac" Value="('$(TargetPlatform)' == 'Mac' or '$(TargetPlatform)' == 'IOS' or '$(TargetPlatform)' == 'tvOS')"/>
+				<Property Name="IsTargetPlatformLinux" Value="('$(TargetPlatform)' == 'Linux')"/>
+				<Property Name="IsTargetPlatformWindows" Value="(!$(IsTargetPlatformMac) and !$(IsTargetPlatformLinux))"/>
+
+				<!-- Declare type of host required to build & cook this platform-->
+				<Property Name="HostAgentType" Value="$(EditorPlatformWin64AgentOverride)" /> 
+				<Property Name="HostAgentType" Value="$(EditorPlatformMacAgentOverride)" If="$(IsTargetPlatformMac)"/>
+				<Property Name="HostAgentType" Value="$(EditorPlatformLinuxAgentOverride)" If="$(IsTargetPlatformLinux)"/>
+
+				<Property Name="IsEditorAgentOverride" Value="false"/>
+				<Property Name="IsEditorAgentOverride" Value="true" If="$(IsTargetPlatformWindows) and '$(EditorPlatformWin64AgentOverride)'!='Win64'"/>
+				<Property Name="IsEditorAgentOverride" Value="true" If="$(IsTargetPlatformMac) and '$(EditorPlatformMacAgentOverride)'!='Mac'"/>
+				<Property Name="IsEditorAgentOverride" Value="true" If="$(IsTargetPlatformLinux) and '$(EditorPlatformLinuxAgentOverride)'!='Linux'"/>
+				<!-- On consoles we can use agents with devkit setup, but locally hosted tests need a Test agent-->
+				<Property Name="TestAgentType" Value="$(TargetPlatformTestAgentOverride)" />
+				<Property Name="TestAgentType" Value="Test$(HostAgentType)" If="$(IsDesktopPlatform)"/>
+				<Property Name="TestAgentType" Value="$(HostAgentType)" If="$(IsDesktopPlatform) and $(DontUseTestAgent)" />
+				<Property Name="TestAgentType" Value="$(TestDesktopPlatformWin64AgentOverride)" If="$(IsDesktopPlatform) and '$(TestDesktopPlatformWin64AgentOverride)' != '' and $(IsTargetPlatformWindows)"/>
+				<Property Name="TestAgentType" Value="$(TestDesktopPlatformLinuxAgentOverride)" If="$(IsDesktopPlatform) and '$(TestDesktopPlatformLinuxAgentOverride)' != '' and $(IsTargetPlatformLinux)"/>
+				<Property Name="TestAgentType" Value="$(TestDesktopPlatformMacAgentOverride)" If="$(IsDesktopPlatform) and '$(TestDesktopPlatformMacAgentOverride)' != '' and $(IsTargetPlatformMac)"/>
+
+				<!-- Use incremental agents for building? -->
+				<Property Name="CompileAgentType" Value="$(HostAgentType)" />
+				<Property Name="CompileAgentType" Value="Incremental$(HostAgentType)" If="$(UseIncrementalAgents)" />
+
+				<!-- Use agent types from override? -->
+				<Property Name="CompileAgentType" Value="$(AgentOverride)" If="'$(AgentOverride)' != ''"/>
+				<!-- Use the same agent type for packaging as compilation -->
+				<Property Name="PackageAgentType" Value="$(CompileAgentType)"/>
+				<Property Name="CompileAgentType" Value="$(ClientCompileAgentOverride)" If="'$(ClientCompileAgentOverride)' != ''"/>
+				<Property Name="TestAgentType" Value="$(TestAgentOverride)" If="'$(TestAgentOverride)' != ''"/>
+
+				<Property Name="PreNodeNameWithPlatform" Value="$(PreNodeName)"/>
+				<Property Name="PreNodeNameWithPlatform" Value="[$(HostAgentType)]$(PreNodeName)" If="$(IsEditorAgentOverride)"/>
+				<!-- Building and cooking will require these nodes from the appropriate host. They should have been built above... -->
+				<Property Name="PlatformToolsNodeName" Value="$(PreNodeName)Compile Tools $(HostAgentType)" />
+				<Property Name="PlatformEditorNodeName" Value="$(PreNodeName)Compile Editor $(HostAgentType)" />
+				<!-- Define these names up front as we use some of them several times -->
+				<Property Name="PlatformAgentName" Value="$(ProjectName) $(TargetPlatform)$(VariationTitle)" />
+				<Property Name="PlatformCompileNodeName" Value="$(PreNodeNameWithPlatform)Compile $(TargetPlatform)$(VariationTitle)" />
+				<Property Name="UploadSymbolsNodeName" Value="" />
+				<Property Name="UploadSymbolsNodeName" Value="$(PreNodeNameWithPlatform)Upload Symbols $(TargetPlatform)$(VariationTitle)" If="'$(SymbolServerPath)' != '' And !$(IsMobilePlatform) And !$(IsTargetPlatformLinux)" />
+				<Property Name="PlatformCookNodeName" Value="$(PreNodeNameWithPlatform)Cook $(TargetPlatform)$(VariationTitle)" />
+				<Property Name="PlatformStageNodeName" Value="$(PreNodeNameWithPlatform)Stage $(TargetPlatform)$(VariationTitle)" />
+				<Property Name="PlatformPackageNodeName" Value="$(PreNodeNameWithPlatform)Package $(TargetPlatform)$(VariationTitle)" />
+				<Property Name="PlatformPackageArtifactName" Value="$(PlatformPackageNodeName)" />
+				<Property Name="PlatformHashNodeName" Value="$(PreNodeNameWithPlatform)Hash $(TargetPlatform)$(VariationTitle)" />
+				<Property Name="PlatformPublishStagedNodeName" Value="$(PreNodeNameWithPlatform)Publish Staged $(TargetPlatform)$(VariationTitle)" />
+				<Property Name="PlatformPublishPackagedNodeName" Value="$(PreNodeNameWithPlatform)Publish Packaged $(TargetPlatform)$(VariationTitle)" />
+				<Property Name="PlatformCheckBuildSizeNodeName" Value="$(PreNodeNameWithPlatform)Check Build Size $(TargetPlatform)$(VariationTitle)" />
+				<Property Name="PlatformStagedArtifactNodeName" Value="$(PlatformPublishStagedNodeName) (Horde)" />
+				<Property Name="CloudPlatformStagedArtifactNodeName" Value="$(PlatformPublishStagedNodeName) (Cloud)" />
+				<Property Name="HordeArtifactPlatformPackagedNodeName" Value="$(PlatformPublishPackagedNodeName) (Horde)" />
+				<Property Name="HordeArtifactPlatformRetrieve" Value="$(PreNodeNameWithPlatform)Retrieve Build Artifact $(TargetPlatform)$(VariationTitle)" />
+				<Property Name="CloudArtifactPlatformPackagedNodeName" Value="$(PlatformPublishPackagedNodeName) (Cloud)" />
+				<Property Name="CloudArtifactPlatformRetrieve" Value="$(PreNodeNameWithPlatform)Retrieve Build Artifact $(TargetPlatform)$(VariationTitle)" />
+				<Property Name="RequirePlatformStaged" Value="$(PlatformStageNodeName)" />
+				<Property Name="RequirePlatformStaged" Value="$(PlatformStageNodeName) Server" If="$(DedicatedServer)"/>
+				<Property Name="RequirePlatformStaged" Value="$(RequirePlatformStaged);$(PlatformStagedArtifactNodeName)" If="!$(SkipPublish) And $(PublishHordeArtifacts)"/>
+				<Property Name="RequirePlatformStaged" Value="$(RequirePlatformStaged);$(PlatformStagedArtifactNodeName) Server" If="!$(SkipPublish) And $(PublishHordeArtifacts) And $(DedicatedServer)"/>
+				<Property Name="RequirePlatformStaged" Value="$(RequirePlatformStaged);$(CloudPlatformStagedArtifactNodeName)" If="!$(SkipPublish) And $(PublishCloudArtifacts)"/>
+				<Property Name="RequirePlatformStaged" Value="$(RequirePlatformStaged);$(PlatformPublishStagedNodeName)" If="!$(SkipPublish) And $(PublishBuilds)"/>
+				<Property Name="RequirePlatformStaged" Value="$(RequirePlatformStaged);$(PlatformPublishStagedNodeName) Server" If="!$(SkipPublish) And $(PublishBuilds) And $(DedicatedServer)"/>
+				<Property Name="RequirePlatformPackaged" Value="$(PlatformPackageNodeName)"/>
+				<Property Name="RequirePlatformPackaged" Value="$(PlatformPackageNodeName) Server" If="$(DedicatedServer)"/>
+				<Property Name="RequirePlatformPackaged" Value="$(RequirePlatformPackaged);$(HordeArtifactPlatformPackagedNodeName)" If="!$(SkipPublish) And $(PublishHordeArtifacts)"/>
+				<Property Name="RequirePlatformPackaged" Value="$(RequirePlatformPackaged);$(HordeArtifactPlatformPackagedNodeName) Server" If="!$(SkipPublish) And $(PublishHordeArtifacts) And $(DedicatedServer)"/>
+				<Property Name="RequirePlatformPackaged" Value="$(RequirePlatformPackaged);$(CloudArtifactPlatformPackagedNodeName)" If="!$(SkipPublish) And $(PublishCloudArtifacts)"/>
+				<Property Name="RequirePlatformPackaged" Value="$(RequirePlatformPackaged);$(PlatformPublishPackagedNodeName)"  If="!$(SkipPublish) And $(PublishBuilds)"/>
+				<Property Name="RequirePlatformPackaged" Value="$(RequirePlatformPackaged);$(PlatformPublishPackagedNodeName) Server"  If="!$(SkipPublish) And $(PublishBuilds) And $(DedicatedServer)"/>
+				
+				<!-- Set CookPlatform since it isn't always the same as the target -->
+				<Property Name="CookPlatform" Value="$(TargetPlatform)" />
+				<Property Name="CookPlatform" Value="Windows" If="'$(TargetPlatform)' == 'Win64'"/>
+				<Property Name="CookPlatform" Value="$(CookPlatform)_$(AndroidCookFlavor)" If="'$(TargetPlatform)' == 'Android'"/>
+				<Property Name="CookPlatform" Value="$(CookPlatform)Client" If="'$(TargetType)' == 'Client'" />
+
+				<!-- Only use Cook Flavor for Android -->
+				<Property Name="BCRCookFlavor" Value=""/>
+				<Property Name="BCRCookFlavor" Value="$(AndroidCookFlavor)" If="'$(TargetPlatform)' == 'Android'" />
+				
+				<Property Name="StagedPlatformFolder" Value="$(CookPlatform)"  />
+				<!-- <Property Name="StagedPlatformFolder" Value="$(StagedPlatformFolder)Client" If="'$(TargetType)' == 'Client'" /> -->
+
+				<!-- Hashing of staged data can only be done on Win64 hosts -->
+				<Property Name="DoUnsyncHashing" Value="$(HostAgentType) == 'Win64' and $(IsEpicInternal)" />
+
+				<!-- Default BuildCookRun args needed when staging/packaging/archiving this configuration -->
+				<Property Name="BCRArgs" Value="-project=$(TargetProject) -platform=$(TargetPlatform) -cookflavor=&quot;$(BCRCookFlavor)&quot;" />
+				<Property Name="BCRArgs" Value="$(BCRArgs) -client" If="'$(TargetType)' == 'Client'" />
+				<Property Name="BCRArgs" Value="$(BCRArgs) -clientarchitecture=arm64" If="$(IsArm64Variation)"/>
+				<Property Name="BCRArgs" Value="$(BCRArgs) -clientarchitecture=arm64ec" If="$(IsArm64ecVariation)"/>
+				<Property Name="BCRArgs" Value="$(BCRArgs) -clientarchitecture=x64+arm64" If="$(IsMultiArchVariation)"/>
+
+				<Property Name="CompileArguments" Value="$(GenericCompileArguments)"/>
+				<Property Name="CompileArguments" Value="$(CompileArguments) $(ExtraWin64CompileArguments)" If="$(IsDesktopPlatform) and '$(ExtraWin64CompileArguments)' != '' and $(IsTargetPlatformWindows)"/>
+				<Property Name="CompileArguments" Value="$(CompileArguments) $(ExtraLinuxCompileArguments)" If="$(IsDesktopPlatform) and '$(ExtraLinuxCompileArguments)' != '' and $(IsTargetPlatformLinux)"/>
+				<Property Name="CompileArguments" Value="$(CompileArguments) $(ExtraMacCompileArguments)" If="$(IsDesktopPlatform) and '$(ExtraMacCompileArguments)' != '' and $(IsTargetPlatformMac)"/>
+				<Property Name="CompileArguments" Value="$(CompileArguments) -architecture=arm64" If="$(IsArm64Variation)"/>
+				<Property Name="CompileArguments" Value="$(CompileArguments) -architecture=arm64ec" If="$(IsArm64ecVariation)"/>
+				<Property Name="CompileArguments" Value="$(CompileArguments) -architecture=x64+arm64" If="$(IsMultiArchVariation)"/>
+
+				<!-- Support additional stage and package arguments that are only active on some platforms -->
+				<Property Name="StageAndPackageArguments" Value="$(ExtraStageAndPackageArguments)" />
+				<Property Name="StageAndPackageArguments" Value="$(StageAndPackageArguments) $(TargetPlatformFilteredExtraStageAndPackageArguments)" If="'$(TargetPlatformFilteredExtraStageAndPackageArgumentsFilter)' != '' and ContainsItem('$(TargetPlatformFilteredExtraStageAndPackageArgumentsFilter)', '$(TargetPlatform)', '+')" />
+				<Property Name="PackageOnlyArguments" Value="$(ExtraPackageOnlyArguments)" />
+
+				<Property Name="StageAndPackageServerArguments" Value="$(ExtraStageAndPackageServerArguments)" />
+
+				<Property Name="BuildSizeLimit" Value=""/>
+				<Regex Pattern="(?&lt;!\w)$(TargetPlatform):([0-9]+)" Capture="BuildSizeLimit" Input="$(CheckBuildSizePlatforms)"/>
+
+				<!-- Build the client and if necessary a server-->
+				<Agent Name="Compile $(PlatformAgentName)" Type="$(CompileAgentType)">
+
+					<!-- Compile -->
+					<Node Name="$(PlatformCompileNodeName)" Requires="$(PlatformToolsNodeName)">
+						<Switch>
+							<Case If="$(IsContentOnlyProject)">
+								<Command Name="BuildGame" Arguments="$(BCRArgs) -configuration=$(TargetConfigurations) -notools -ubtargs=&quot;-BuildVersion=$(BuildVersion) $(CompileArguments) $(TargetCompileArguments)&quot;" />
+								<Tag Files="$(ProjectPath)/Binaries/..." Filter="$(TargetName)*.target" With="#$(PlatformCompileNodeName)"/>
+								<Tag Files="$(ProjectPath)/Binaries/..." Filter="UnrealGame*.target" With="#UnrealGameTarget"/>
+								<Copy Files="#UnrealGameTarget" From="$(ProjectPath)" To="Engine" Tag="#$(PlatformCompileNodeName)"/>
+								<TagReceipt Files="#$(PlatformCompileNodeName)" BuildProducts="true" With="#$(PlatformCompileNodeName)"/>
+							</Case>
+							<Default>
+								<ForEach Name="TargetConfiguration" Values="$(TargetConfigurations)" Separator="+">
+									<Compile Target="$(TargetName)" Project="$(TargetProject)" Platform="$(TargetPlatform)" Configuration="$(TargetConfiguration)" Arguments="-BuildVersion=&quot;$(BuildVersion)&quot; $(CompileArguments) $(TargetCompileArguments)"/>
+								</ForEach>
+							</Default>
+						</Switch>
+					</Node>
+
+					<Node Name="$(PlatformCompileNodeName) Server" Requires="$(PlatformToolsNodeName)">
 						<ForEach Name="TargetConfiguration" Values="$(TargetConfigurations)" Separator="+">
-							<Compile Target="$(TargetName)" Project="$(TargetProject)" Platform="$(TargetPlatform)" Configuration="$(TargetConfiguration)" Arguments="-BuildVersion=&quot;$(BuildVersion)&quot; $(CompileArguments) $(TargetCompileArguments)"/>
-						</ForEach>
-					</Default>
-				</Switch>
-			</Node>
-
-			<Node Name="$(PlatformCompileNodeName) Server" Requires="$(PlatformToolsNodeName)">
-				<ForEach Name="TargetConfiguration" Values="$(TargetConfigurations)" Separator="+">
-					<Property Name="ServerPlatform" Value="Win64"/>
-					<Property Name="ServerPlatform" Value="Linux" If="$(IsTargetPlatformLinux)"/>
-					<Property Name="ServerCompileArguments" Value="-BuildVersion=&quot;$(BuildVersion)&quot; $(CompileArguments) $(TargetCompileArguments)"/>
-					<Compile Target="$(TargetNameServer)" Platform="$(ServerPlatform)" Configuration="$(TargetConfiguration)" Arguments="$(ServerCompileArguments)"/>
-				</ForEach>	
-			</Node>
-
-			<!-- (Optional) Upload Symbols -->
-			<Do If="'$(UploadSymbolsNodeName)' != ''">
-				<Node Name="$(UploadSymbolsNodeName)" Requires="$(PlatformCompileNodeName)">
-					<Log Message="Uploading symbols to &quot;$(SymbolServerPath)&quot; ..."/>
-					<SymStore Platform="$(TargetPlatform)" Files="#$(PlatformCompileNodeName)" StoreDir="$(SymbolServerPath)/$(TargetPlatform)" Product="$(TargetName)" BuildVersion="$(BuildVersion)" />
-				</Node>
-			</Do>
-
-		</Agent>
-
-		<Property Name="PlatformCookAdditionalRequirements" Value=""/>
-		<Expand Name="CookAdditionalRequirements"/>
-
-		<Property Name="PlatformEditorNodeName" Value="" If="$(SkipEditorBuild)" />
-		<Property Name="PlatformSnapshotPublishNodeName" Value="$(PreNodeName)Snapshot Publish $(TargetPlatform)" />
-		<!-- Separate agent for cooking / staging, so the cook can happen in parallel with the client compile -->
-		<Agent Name="Package $(PlatformAgentName)" Type="$(PackageAgentType)">
-
-			<!-- Cook -->
-			<Node Name="$(PlatformCookNodeName)" Requires="$(PlatformEditorNodeName);$(PlatformCookAdditionalRequirements)" Produces="#$(PlatformCookNodeName) Complete;#$(PlatformCookNodeName) CookShaderSymbols" RunEarly="true">
-				<Do If="!$(SkipCook)">
-					<Expand Name="PreCookStep"/>
-					<Property Name="CookArgs" Value="" />
-					<Property Name="CookArgs" Value="-MapIniSection=&quot;$(MapIniSection)&quot; $(CookPlatform)" If="'$(MapIniSection)' != ''"/>				
-					<Cook Project="$(TargetProject)" Platform="$(CookPlatform)" TagOutput="false" Arguments="$(CookArgs) $(ExtraCookArguments)"/>
-
-					<!-- If Reshape, keep the symbols around for later -->
-					<Do If="$(DoGPUReshape)">
-						<Tag Files="$(RootDir)\$(ProjectPath)\Saved\ShaderSymbols\..." With="#$(PlatformCookNodeName) CookShaderSymbols"/>
-						<Property Name="RequirePlatformStaged" Value="$(RequirePlatformStaged);#$(PlatformCookNodeName) CookShaderSymbols"/>
+							<Property Name="ServerPlatform" Value="Win64"/>
+							<Property Name="ServerPlatform" Value="Linux" If="$(IsTargetPlatformLinux)"/>
+							<Property Name="ServerCompileArguments" Value="-BuildVersion=&quot;$(BuildVersion)&quot; $(CompileArguments) $(TargetCompileArguments)"/>
+							<Compile Target="$(TargetNameServer)" Platform="$(ServerPlatform)" Configuration="$(TargetConfiguration)" Arguments="$(ServerCompileArguments)"/>
+						</ForEach>	
+					</Node>
+
+					<!-- (Optional) Upload Symbols -->
+					<Do If="'$(UploadSymbolsNodeName)' != ''">
+						<Node Name="$(UploadSymbolsNodeName)" Requires="$(PlatformCompileNodeName)">
+							<Log Message="Uploading symbols to &quot;$(SymbolServerPath)&quot; ..."/>
+							<SymStore Platform="$(TargetPlatform)" Files="#$(PlatformCompileNodeName)" StoreDir="$(SymbolServerPath)/$(TargetPlatform)" Product="$(TargetName)" BuildVersion="$(BuildVersion)" />
+						</Node>
 					</Do>
-				</Do>
-			</Node>
-			
-			<Node Name="$(PlatformCookNodeName) Server" Requires="$(PlatformEditorNodeName);$(PlatformCookAdditionalRequirements)" Produces="#$(PlatformCookNodeName) Server Complete" RunEarly="true">
-				<Do If="!$(SkipCook)">
-					<Expand Name="PreCookStep"/>
-					<Property Name="CookArgs" Value="" />
-					<Property Name="CookArgs" Value="-MapIniSection=&quot;$(MapIniSection)&quot; $(CookPlatform)" If="'$(MapIniSection)' != ''"/>	
 
-					<Property Name="ServerPlatformCook" Value="Windows"/>
-					<Property Name="ServerPlatformCook" Value="Linux" If="$(IsTargetPlatformLinux)"/>
-					<Property Name="ServerPlatformCook" Value="$(ServerPlatformCook)Server"/>
+				</Agent>
 
-					<Cook Project="$(TargetProject)" Platform="$(ServerPlatformCook)" TagOutput="false" Arguments="$(CookArgs) $(ExtraCookArguments)"/>
-				</Do>
-			</Node>
-			
-			<!-- Stage -->
-			<Node Name="$(PlatformStageNodeName)" Requires="$(PlatformCompileNodeName);#$(PlatformCookNodeName) Complete">
-				<Switch>
-					<Case If="$(IsMobilePlatform)">
-						<ForEach Name="PackageConfiguration" Values="$(PackageConfigurations)" Separator="+">
-							<Command Name="BuildCookRun" Arguments="$(BCRArgs) -configuration=$(PackageConfiguration) $(StageAndPackageArguments) -skipbuild -skipcook -stage -pak" />
-						</ForEach>
-					</Case>
-					<Default>
-						<Property Name="ClientStageArguments" Value="$(BCRArgs) -configuration=$(PackageConfigurations) $(StageAndPackageArguments) -skipbuild -skipcook -stage -pak -NoCodeSign" />
-						<Command Name="BuildCookRun" Arguments="$(ClientStageArguments)" />
-					</Default>
-				</Switch>
-			</Node>
+				<Property Name="PlatformCookAdditionalRequirements" Value=""/>
+				<Expand Name="CookAdditionalRequirements"/>
+
+				<Property Name="PlatformEditorNodeName" Value="" If="$(SkipEditorBuild)" />
+				<Property Name="PlatformSnapshotPublishNodeName" Value="$(PreNodeName)Snapshot Publish $(TargetPlatform)$(VariationTitle)" />
+				<!-- Separate agent for cooking / staging, so the cook can happen in parallel with the client compile -->
+				<Agent Name="Package $(PlatformAgentName)" Type="$(PackageAgentType)">
+
+					<!-- Cook -->
+					<Node Name="$(PlatformCookNodeName)" Requires="$(PlatformEditorNodeName);$(PlatformCookAdditionalRequirements)" Produces="#$(PlatformCookNodeName) Complete;#$(PlatformCookNodeName) CookShaderSymbols" RunEarly="true">
+						<Do If="!$(SkipCook)">
+							<Expand Name="PreCookStep"/>
+							<Property Name="CookArgs" Value="" />
+							<Property Name="CookArgs" Value="-MapIniSection=&quot;$(MapIniSection)&quot; $(CookPlatform)" If="'$(MapIniSection)' != ''"/>				
+							<Cook Project="$(TargetProject)" Platform="$(CookPlatform)" TagOutput="false" Arguments="$(CookArgs) $(ExtraCookArguments)"/>
+
+							<!-- If Reshape, keep the symbols around for later -->
+							<Do If="$(DoGPUReshape)">
+								<Tag Files="$(RootDir)\$(ProjectPath)\Saved\ShaderSymbols\..." With="#$(PlatformCookNodeName) CookShaderSymbols"/>
+								<Property Name="RequirePlatformStaged" Value="$(RequirePlatformStaged);#$(PlatformCookNodeName) CookShaderSymbols"/>
+							</Do>
+						</Do>
+					</Node>
+					
+					<Node Name="$(PlatformCookNodeName) Server" Requires="$(PlatformEditorNodeName);$(PlatformCookAdditionalRequirements)" Produces="#$(PlatformCookNodeName) Server Complete" RunEarly="true">
+						<Do If="!$(SkipCook)">
+							<Expand Name="PreCookStep"/>
+							<Property Name="CookArgs" Value="" />
+							<Property Name="CookArgs" Value="-MapIniSection=&quot;$(MapIniSection)&quot; $(CookPlatform)" If="'$(MapIniSection)' != ''"/>	
 
-			<Node Name="$(PlatformStageNodeName) Server" Requires="$(PlatformCompileNodeName) Server;#$(PlatformCookNodeName) Server Complete">
-				<Property Name="BCRArgsServer" Value="$(BCRArgs) -server" />
-				<Property Name="ServerStageArguments" Value="$(BCRArgsServer) -configuration=$(PackageConfigurations) $(StageAndPackageServerArguments) -skipbuild -skipcook -stage -pak -NoCodeSign" />
-				<Command Name="BuildCookRun" Arguments="$(ServerStageArguments)" />
-			</Node>
-			
-			<!-- Hash -->
-			<Node Name="$(PlatformHashNodeName)" Requires="$(PlatformStageNodeName)">
-				<Do If="$(DoUnsyncHashing)">
-					<Spawn Exe="$(RootDir)/Engine/Restricted/NotForLicensees/Binaries/Win64/unsync.exe" Arguments="hash &quot;$(ProjectOutputDirectory)/$(StagedPlatformFolder)&quot; -v --buffered-files"/>
-				</Do>
-			</Node>
+							<Property Name="ServerPlatformCook" Value="Windows"/>
+							<Property Name="ServerPlatformCook" Value="Linux" If="$(IsTargetPlatformLinux)"/>
+							<Property Name="ServerPlatformCook" Value="$(ServerPlatformCook)Server"/>
 
-			<!-- (Optional) Publish from local Zen to Cloud or remote Zen server -->
-			<Property Name="PlatformPublishStageDependencyName" Value="$(PlatformHashNodeName)" />
-			<Property Name="ShouldSnapshotPublishCloud" Value="false"/>
-			<Property Name="ShouldSnapshotPublishCloud" Value="true" If="$(PublishCloudSnapshot) and '$(UE-CloudPublishHost)' != ''"/>
-			<Property Name="ShouldSnapshotPublishZen" Value="false"/>
-			<Property Name="ShouldSnapshotPublishZen" Value="true" If="$(IsBuildMachine) and $(PublishZenSnapshot) and '$(UE-ZenPublishHost)' != ''"/>
-			<Property Name="ShouldSnapshotPublish" Value="false"/>
-			<Property Name="ShouldSnapshotPublish" Value="true" If="$(IsBuildMachine) and $(ShouldSnapshotPublishCloud) or $(ShouldSnapshotPublishZen)"/>
-			<Do If="$(PublishBuilds) and $(ShouldSnapshotPublish)">
-				<Node Name="$(PlatformSnapshotPublishNodeName)" Requires="$(PlatformHashNodeName)">
-					<Property Name="SnapshotLocalDir" Value="$(RootDir)/Snapshots" />
-					<Property Name="SnapshotIdentifier" Value="$(EscapedBranch)/$(BuildNamePath)"/>
-					<Property Name="SnapshotIdentifier" Value="$(UE_HORDE_TEMPLATEID).$(BuildVersion).$(UE_HORDE_JOBID:-6).$(UE_HORDE_STEPID)" If="'$(UE_HORDE_JOBID)' != '' and '$(UE_HORDE_STEPID)' != '' and '$(UE_HORDE_TEMPLATEID)' != ''"/>
-					<Property Name="SnapshotFilenamePrefix" Value="SnapshotDescriptor"/>
-					<Property Name="SnapshotFilenamePrefix" Value="$(UE_HORDE_TEMPLATEID)" If="'$(UE_HORDE_TEMPLATEID)' != ''"/>
-					<Property Name="SnapshotLongTermPlatformFolder" Value="$(StagedPlatformFolder)" />
-					<Property Name="SnapshotLongTermPlatformFolder" Value="Android" If="'$(TargetPlatform)' == 'Android'" />
-					<Property Name="SnapshotLongTermStorageDir" Value="$(NetworkSnapshotDescFileRoot)/$(EscapedBranch)/$(ProjectName)/$(SnapshotLongTermPlatformFolder)/$(Change)"  If="'$(NetworkSnapshotDescFileRoot)' != ''"/>
-
-					<Do If="$(ShouldSnapshotPublishCloud)">
-						<Property Name="SnapshotBranch" Value="$(EscapedBranch)"/>
-						<StringOp Input="$(EscapedBranch)" Method="Replace" Arguments=".;-" Output="SnapshotBranch"/>
-						<!-- Publish to builds cloud storage endpoint -->
-						<Property Name="SnapshotBuildBaselineBranch" Value="$(Branch)"/>
-						<Property Name="SnapshotBuildBaselineBranch" Value="$(UE-CloudPublishBaselineBranch)" If="'$(UE-CloudPublishBaselineBranch)' != ''"/>
-						<Property Name="SnapshotBuildMetaData" Value="name=$(CookPlatform);buildGroup=$(BuildVersion);depot=$(Depot);branch=$(Branch);baselineBranch=$(SnapshotBuildBaselineBranch);project=$(ProjectName);runtime=$(TargetType);platform=$(TargetPlatform);cookPlatform=$(CookPlatform);changelist=$(Change);commit=$(Change);codeChangelist=$(CodeChange);codeCommit=$(CodeChange);isPreflight=$(IsPreflight);isBuildMachine=$(IsBuildMachine);isHordeEnv=$(IsHordeEnv);buildgroup=$(BuildVersion)"/>
-						<Property Name="SnapshotBuildMetaData" Value="$(SnapshotBuildMetaData);cookFlavor=$(BCRCookFlavor)" If="'$(BCRCookFlavor)' != ''" />
-						<Property Name="SnapshotBuildMetaData" Value="$(SnapshotBuildMetaData);preflightChangelist=$(PreflightChange)" If="'$(PreflightChange)' != ''" />
-						<Property Name="SnapshotBuildMetaData" Value="$(SnapshotBuildMetaData);hordeTemplateId=$(UE_HORDE_TEMPLATEID)" If="'$(UE_HORDE_TEMPLATEID)' != ''" />
-						<Log Message="Publishing build with metadata '$(SnapshotBuildMetaData)'"/>
-						<ZenExportSnapshot Project="$(TargetProject)" Platform="$(CookPlatform)" DestinationStorageType="Builds" DestinationCloudHost="$(UE-CloudPublishHost)" DestinationCloudHttpVersion="$(UE-CloudPublishHttpVersion)" SnapshotDescriptorCloudHost="$(UE-CloudPublishDescriptorHost)" SnapshotDescriptorCloudHttpVersion="$(UE-CloudPublishDescriptorHttpVersion)" DestinationCloudNamespace="$(UE-CloudPublishNamespace)" DestinationCloudBucket="$(ProjectName).oplog.$(SnapshotBranch).$(CookPlatform)" DestinationIdentifier="$(SnapshotIdentifier)" SnapshotDescriptorFile="$(SnapshotLocalDir)/$(StagedPlatformFolder)/$(SnapshotFilenamePrefix)-builds.json" Metadata="$(SnapshotBuildMetaData)"/>
-						<Copy From="$(SnapshotLocalDir)/$(StagedPlatformFolder)/$(SnapshotFilenamePrefix)-builds.json" To="$(NetworkOutputDirectory)/$(StagedPlatformFolder)/SnapshotDescriptor.json"/>
-						<Do If="'$(SnapshotLongTermStorageDir)' != ''">
-							<Copy From="$(SnapshotLocalDir)/$(StagedPlatformFolder)/$(SnapshotFilenamePrefix)-builds.json" To="$(SnapshotLongTermStorageDir)/$(SnapshotFilenamePrefix)-builds.json"/>
+							<Cook Project="$(TargetProject)" Platform="$(ServerPlatformCook)" TagOutput="false" Arguments="$(CookArgs) $(ExtraCookArguments)"/>
+						</Do>
+					</Node>
+					
+					<!-- Stage -->
+					<Node Name="$(PlatformStageNodeName)" Requires="$(PlatformCompileNodeName);#$(PlatformCookNodeName) Complete">
+						<Switch>
+							<Case If="$(IsMobilePlatform)">
+								<ForEach Name="PackageConfiguration" Values="$(PackageConfigurations)" Separator="+">
+									<Command Name="BuildCookRun" Arguments="$(BCRArgs) -configuration=$(PackageConfiguration) $(StageAndPackageArguments) -skipbuild -skipcook -stage -pak" />
+								</ForEach>
+							</Case>
+							<Default>
+								<Property Name="ClientStageArguments" Value="$(BCRArgs) -configuration=$(PackageConfigurations) $(StageAndPackageArguments) -skipbuild -skipcook -stage -pak -NoCodeSign" />
+								<Command Name="BuildCookRun" Arguments="$(ClientStageArguments)" />
+							</Default>
+						</Switch>
+					</Node>
+
+					<Node Name="$(PlatformStageNodeName) Server" Requires="$(PlatformCompileNodeName) Server;#$(PlatformCookNodeName) Server Complete">
+						<Property Name="BCRArgsServer" Value="$(BCRArgs) -server" />
+						<Property Name="ServerStageArguments" Value="$(BCRArgsServer) -configuration=$(PackageConfigurations) $(StageAndPackageServerArguments) -skipbuild -skipcook -stage -pak -NoCodeSign" />
+						<Command Name="BuildCookRun" Arguments="$(ServerStageArguments)" />
+					</Node>
+					
+					<!-- Hash -->
+					<Node Name="$(PlatformHashNodeName)" Requires="$(PlatformStageNodeName)">
+						<Do If="$(DoUnsyncHashing)">
+							<Spawn Exe="$(RootDir)/Engine/Restricted/NotForLicensees/Binaries/Win64/unsync.exe" Arguments="hash &quot;$(ProjectOutputDirectory)/$(StagedPlatformFolder)&quot; -v --buffered-files"/>
 						</Do>
+					</Node>
+
+					<!-- (Optional) Publish from local Zen to Cloud or remote Zen server -->
+					<Property Name="PlatformPublishStageDependencyName" Value="$(PlatformHashNodeName)" />
+					<Property Name="ShouldSnapshotPublishCloud" Value="false"/>
+					<Property Name="ShouldSnapshotPublishCloud" Value="true" If="$(PublishCloudSnapshot) and '$(UE-CloudPublishHost)' != ''"/>
+					<Property Name="ShouldSnapshotPublishZen" Value="false"/>
+					<Property Name="ShouldSnapshotPublishZen" Value="true" If="$(IsBuildMachine) and $(PublishZenSnapshot) and '$(UE-ZenPublishHost)' != ''"/>
+					<Property Name="ShouldSnapshotPublish" Value="false"/>
+					<Property Name="ShouldSnapshotPublish" Value="true" If="$(IsBuildMachine) and $(ShouldSnapshotPublishCloud) or $(ShouldSnapshotPublishZen)"/>
+					<Do If="$(PublishBuilds) and $(ShouldSnapshotPublish)">
+						<Node Name="$(PlatformSnapshotPublishNodeName)" Requires="$(PlatformHashNodeName)">
+							<Property Name="SnapshotLocalDir" Value="$(RootDir)/Snapshots" />
+							<Property Name="SnapshotIdentifier" Value="$(EscapedBranch)/$(BuildNamePath)"/>
+							<Property Name="SnapshotIdentifier" Value="$(UE_HORDE_TEMPLATEID).$(BuildVersion).$(UE_HORDE_JOBID:-6).$(UE_HORDE_STEPID)" If="'$(UE_HORDE_JOBID)' != '' and '$(UE_HORDE_STEPID)' != '' and '$(UE_HORDE_TEMPLATEID)' != ''"/>
+							<Property Name="SnapshotFilenamePrefix" Value="SnapshotDescriptor"/>
+							<Property Name="SnapshotFilenamePrefix" Value="$(UE_HORDE_TEMPLATEID)" If="'$(UE_HORDE_TEMPLATEID)' != ''"/>
+							<Property Name="SnapshotLongTermPlatformFolder" Value="$(StagedPlatformFolder)" />
+							<Property Name="SnapshotLongTermPlatformFolder" Value="Android" If="'$(TargetPlatform)' == 'Android'" />
+							<Property Name="SnapshotLongTermStorageDir" Value="$(NetworkSnapshotDescFileRoot)/$(EscapedBranch)/$(ProjectName)/$(SnapshotLongTermPlatformFolder)/$(Change)"  If="'$(NetworkSnapshotDescFileRoot)' != ''"/>
+
+							<Do If="$(ShouldSnapshotPublishCloud)">
+								<Property Name="SnapshotBranch" Value="$(EscapedBranch)"/>
+								<StringOp Input="$(EscapedBranch)" Method="Replace" Arguments=".;-" Output="SnapshotBranch"/>
+								<!-- Publish to builds cloud storage endpoint -->
+								<Property Name="SnapshotBuildBaselineBranch" Value="$(Branch)"/>
+								<Property Name="SnapshotBuildBaselineBranch" Value="$(UE-CloudPublishBaselineBranch)" If="'$(UE-CloudPublishBaselineBranch)' != ''"/>
+								<Property Name="SnapshotBuildMetaData" Value="name=$(CookPlatform);buildGroup=$(BuildVersion);depot=$(Depot);branch=$(Branch);baselineBranch=$(SnapshotBuildBaselineBranch);project=$(ProjectName);runtime=$(TargetType);platform=$(TargetPlatform);cookPlatform=$(CookPlatform);changelist=$(Change);commit=$(Change);codeChangelist=$(CodeChange);codeCommit=$(CodeChange);isPreflight=$(IsPreflight);isBuildMachine=$(IsBuildMachine);isHordeEnv=$(IsHordeEnv);buildgroup=$(BuildVersion)"/>
+								<Property Name="SnapshotBuildMetaData" Value="$(SnapshotBuildMetaData);cookFlavor=$(BCRCookFlavor)" If="'$(BCRCookFlavor)' != ''" />
+								<Property Name="SnapshotBuildMetaData" Value="$(SnapshotBuildMetaData);preflightChangelist=$(PreflightChange)" If="'$(PreflightChange)' != ''" />
+								<Property Name="SnapshotBuildMetaData" Value="$(SnapshotBuildMetaData);hordeTemplateId=$(UE_HORDE_TEMPLATEID)" If="'$(UE_HORDE_TEMPLATEID)' != ''" />
+								<Log Message="Publishing build with metadata '$(SnapshotBuildMetaData)'"/>
+								<ZenExportSnapshot Project="$(TargetProject)" Platform="$(CookPlatform)" DestinationStorageType="Builds" DestinationCloudHost="$(UE-CloudPublishHost)" DestinationCloudHttpVersion="$(UE-CloudPublishHttpVersion)" SnapshotDescriptorCloudHost="$(UE-CloudPublishDescriptorHost)" SnapshotDescriptorCloudHttpVersion="$(UE-CloudPublishDescriptorHttpVersion)" DestinationCloudNamespace="$(UE-CloudPublishNamespace)" DestinationCloudBucket="$(ProjectName).oplog.$(SnapshotBranch).$(CookPlatform)" DestinationIdentifier="$(SnapshotIdentifier)" SnapshotDescriptorFile="$(SnapshotLocalDir)/$(StagedPlatformFolder)/$(SnapshotFilenamePrefix)-builds.json" Metadata="$(SnapshotBuildMetaData)"/>
+								<Copy From="$(SnapshotLocalDir)/$(StagedPlatformFolder)/$(SnapshotFilenamePrefix)-builds.json" To="$(NetworkOutputDirectory)/$(StagedPlatformFolder)/SnapshotDescriptor.json"/>
+								<Do If="'$(SnapshotLongTermStorageDir)' != ''">
+									<Copy From="$(SnapshotLocalDir)/$(StagedPlatformFolder)/$(SnapshotFilenamePrefix)-builds.json" To="$(SnapshotLongTermStorageDir)/$(SnapshotFilenamePrefix)-builds.json"/>
+								</Do>
+							</Do>
+							<Do If="$(ShouldSnapshotPublishZen)">
+								<ZenExportSnapshot Project="$(TargetProject)" Platform="$(CookPlatform)" DestinationStorageType="Zen" DestinationZenHost="$(UE-ZenPublishHost)" DestinationIdentifier="$(SnapshotIdentifier)" SnapshotDescriptorFile="$(SnapshotLocalDir)/$(StagedPlatformFolder)/$(SnapshotFilenamePrefix)-zen.json"/>
+								<Copy From="$(SnapshotLocalDir)/$(StagedPlatformFolder)/$(SnapshotFilenamePrefix)-zen.json" To="$(NetworkOutputDirectory)/$(StagedPlatformFolder)/SnapshotDescriptorZen.json"/>
+								<Do If="'$(SnapshotLongTermStorageDir)' != ''">
+									<Copy From="$(SnapshotLocalDir)/$(StagedPlatformFolder)/$(SnapshotFilenamePrefix)-zen.json" To="$(SnapshotLongTermStorageDir)/$(SnapshotFilenamePrefix)-zen.json"/>
+								</Do>
+							</Do>
+							<Tag Files="$(SnapshotLocalDir)/..." With="#SnapshotDescriptor_$(StagedPlatformFolder)"/>
+
+							<CreateArtifact Name="$(ProjectName)-ci-snapshot-$(StagedPlatformFolder)$(VariationTitle)" Description="Snapshot ($(StagedPlatformFolder)$(VariationTitle))" Type="ugs-snapshot" BaseDir="Snapshots" Keys="ugs-project=$(TargetProject)" Metadata="ArchiveType=SnapshotDescriptor_$(StagedPlatformFolder);dashboard-category=Snapshots" Files="#SnapshotDescriptor_$(StagedPlatformFolder)"/>
+						</Node>
+						<Property Name="PlatformPublishStageDependencyName" Value="$(PlatformSnapshotPublishNodeName)" />
 					</Do>
-					<Do If="$(ShouldSnapshotPublishZen)">
-						<ZenExportSnapshot Project="$(TargetProject)" Platform="$(CookPlatform)" DestinationStorageType="Zen" DestinationZenHost="$(UE-ZenPublishHost)" DestinationIdentifier="$(SnapshotIdentifier)" SnapshotDescriptorFile="$(SnapshotLocalDir)/$(StagedPlatformFolder)/$(SnapshotFilenamePrefix)-zen.json"/>
-						<Copy From="$(SnapshotLocalDir)/$(StagedPlatformFolder)/$(SnapshotFilenamePrefix)-zen.json" To="$(NetworkOutputDirectory)/$(StagedPlatformFolder)/SnapshotDescriptorZen.json"/>
-						<Do If="'$(SnapshotLongTermStorageDir)' != ''">
-							<Copy From="$(SnapshotLocalDir)/$(StagedPlatformFolder)/$(SnapshotFilenamePrefix)-zen.json" To="$(SnapshotLongTermStorageDir)/$(SnapshotFilenamePrefix)-zen.json"/>
+					
+					<!-- Publish (Staged) -->
+					<Node Name="$(PlatformPublishStagedNodeName)" Requires="$(PlatformPublishStageDependencyName)">
+						<Log Message="Network upload is disabled, skipping upload." If="!$(PublishToNetworkOutputDirectory)"/>
+						<Do If="$(PublishToNetworkOutputDirectory)">				
+							<Error Message="NetworkOutputDirectory is not set" If="'$(NetworkOutputDirectory)' == ''"/>
+							<Copy Files="..." From="$(ProjectOutputDirectory)/$(StagedPlatformFolder)" To="$(NetworkOutputDirectory)/$(StagedPlatformFolder)/Staged"/>
+							<Copy Files="..." From="$(ShaderSymbolsDirectory)/*$(TargetPlatform)*" To="$(NetworkOutputDirectory)/$(StagedPlatformFolder)/ShaderSymbols"/>
+						</Do>
+					</Node>
+					
+					<Node Name="$(PlatformPublishStagedNodeName) Server" Requires="$(PlatformPublishStageDependencyName)">
+						<Log Message="Network upload is disabled, skipping upload." If="!$(PublishToNetworkOutputDirectory)"/>
+						<Do If="$(PublishToNetworkOutputDirectory)">				
+							<Error Message="NetworkOutputDirectory is not set" If="'$(NetworkOutputDirectory)' == ''"/>
+							<Copy Files="..." From="$(ProjectOutputDirectory)/$(StagedPlatformFolder)Server" To="$(NetworkOutputDirectory)/$(StagedPlatformFolder)Server/Staged"/>
+							<Copy Files="..." From="$(ShaderSymbolsDirectory)/*$(TargetPlatform)Server*" To="$(NetworkOutputDirectory)/$(StagedPlatformFolder)Server/ShaderSymbols"/>
 						</Do>
+					</Node>
+
+					<!-- Publish (Staged) (Cloud) -->
+					<Do If="$(PublishCloudArtifacts)">
+						<Node Name="$(CloudPlatformStagedArtifactNodeName)" Requires="$(PlatformStageNodeName)">
+							<Property Name="TargetDirectory" Value="$(ProjectOutputDirectory)/$(StagedPlatformFolder)"/>
+							<Expand Name="WriteArtifactMetadataFile" Type="staged-build" Platform="$(StagedPlatformFolder)$(VariationTitle)" Configuration="$(PackageConfigurations)" ArtifactDir="$(TargetDirectory)" />
+							<Tag Files="$(TargetDirectory)/..." With="#$(PlatformStageNodeName) Files"/>
+
+							<CreateCloudArtifact 
+								Name="$(ProjectName)-$(TargetPlatform)$(VariationTitle)-Staged" 
+								BuildGroup="$(BuildVersion)" 
+								Type="staged-build" 
+								BaseDir="$(TargetDirectory)" 
+								Metadata="cookPlatform=$(StagedPlatformFolder)$(VariationTitle);dashboard-category=Staged Builds - Cloud;$(SharedCloudBuildMetadata)" 
+								Project="$(TargetProject)" 
+								Platform="$(StagedPlatformFolder)" 
+								Files="#$(PlatformStageNodeName) Files"/>
+							<Delete Files="$(TargetDirectory)/artifactmetadata.json"/>
+						</Node>
 					</Do>
-					<Tag Files="$(SnapshotLocalDir)/..." With="#SnapshotDescriptor_$(StagedPlatformFolder)"/>
 
-					<CreateArtifact Name="$(ProjectName)-ci-snapshot-$(StagedPlatformFolder)" Description="Snapshot ($(StagedPlatformFolder))" Type="ugs-snapshot" BaseDir="Snapshots" Keys="ugs-project=$(TargetProject)" Metadata="ArchiveType=SnapshotDescriptor_$(StagedPlatformFolder);dashboard-category=Snapshots" Files="#SnapshotDescriptor_$(StagedPlatformFolder)"/>
-				</Node>
-				<Property Name="PlatformPublishStageDependencyName" Value="$(PlatformSnapshotPublishNodeName)" />
-			</Do>
-			
-			<!-- Publish (Staged) -->
-			<Node Name="$(PlatformPublishStagedNodeName)" Requires="$(PlatformPublishStageDependencyName)">
-				<Log Message="Network upload is disabled, skipping upload." If="!$(PublishToNetworkOutputDirectory)"/>
-				<Do If="$(PublishToNetworkOutputDirectory)">				
-					<Error Message="NetworkOutputDirectory is not set" If="'$(NetworkOutputDirectory)' == ''"/>
-					<Copy Files="..." From="$(ProjectOutputDirectory)/$(StagedPlatformFolder)" To="$(NetworkOutputDirectory)/$(StagedPlatformFolder)/Staged"/>
-					<Copy Files="..." From="$(ShaderSymbolsDirectory)/*$(TargetPlatform)*" To="$(NetworkOutputDirectory)/$(StagedPlatformFolder)/ShaderSymbols"/>
-				</Do>
-			</Node>
-			
-			<Node Name="$(PlatformPublishStagedNodeName) Server" Requires="$(PlatformPublishStageDependencyName)">
-				<Log Message="Network upload is disabled, skipping upload." If="!$(PublishToNetworkOutputDirectory)"/>
-				<Do If="$(PublishToNetworkOutputDirectory)">				
-					<Error Message="NetworkOutputDirectory is not set" If="'$(NetworkOutputDirectory)' == ''"/>
-					<Copy Files="..." From="$(ProjectOutputDirectory)/$(StagedPlatformFolder)Server" To="$(NetworkOutputDirectory)/$(StagedPlatformFolder)Server/Staged"/>
-					<Copy Files="..." From="$(ShaderSymbolsDirectory)/*$(TargetPlatform)Server*" To="$(NetworkOutputDirectory)/$(StagedPlatformFolder)Server/ShaderSymbols"/>
-				</Do>
-			</Node>
+					<!-- Publish (Staged) (Horde) -->
+					<Do If="$(PublishHordeArtifacts)">
+						<Node Name="$(PlatformStagedArtifactNodeName)" Requires="$(PlatformStageNodeName)">
+							<Property Name="ArtifactsPath" Value="$(ProjectOutputDirectory)/$(StagedPlatformFolder)" />
+							
+							<Expand Name="WriteArtifactMetadataFile" Type="staged-build" Platform="$(StagedPlatformFolder)$(VariationTitle)" Configuration="$(PackageConfigurations)" ArtifactDir="$(ArtifactsPath)"/>
+							<Tag Files="$(ArtifactsPath)/..." With="#$(PlatformStageNodeName) Files"/>
 
-			<!-- Publish (Staged) (Cloud) -->
-			<Do If="$(PublishCloudArtifacts)">
-				<Node Name="$(CloudPlatformStagedArtifactNodeName)" Requires="$(PlatformStageNodeName)">
-					<Property Name="TargetDirectory" Value="$(ProjectOutputDirectory)/$(StagedPlatformFolder)"/>
-					<Expand Name="WriteArtifactMetadataFile" Type="staged-build" Platform="$(StagedPlatformFolder)" Configuration="$(PackageConfigurations)" ArtifactDir="$(TargetDirectory)" />
-					<Tag Files="$(TargetDirectory)/..." With="#$(PlatformStageNodeName) Files"/>
-
-					<CreateCloudArtifact 
-						Name="$(ProjectName)-$(TargetPlatform)-Staged" 
-						BuildGroup="$(BuildVersion)" 
-						Type="staged-build" 
-						BaseDir="$(TargetDirectory)" 
-						Metadata="cookPlatform=$(StagedPlatformFolder);dashboard-category=Staged Builds - Cloud;$(SharedCloudBuildMetadata)" 
-						Project="$(TargetProject)" 
-						Platform="$(StagedPlatformFolder)" 
-						Files="#$(PlatformStageNodeName) Files"/>
-					<Delete Files="$(TargetDirectory)/artifactmetadata.json"/>
-				</Node>
-			</Do>
-
-			<!-- Publish (Staged) (Horde) -->
-			<Do If="$(PublishHordeArtifacts)">
-				<Node Name="$(PlatformStagedArtifactNodeName)" Requires="$(PlatformStageNodeName)">
-					<Property Name="ArtifactsPath" Value="$(ProjectOutputDirectory)/$(StagedPlatformFolder)" />
+							<CreateArtifact Name="$(ProjectName)-Staged-$(StagedPlatformFolder)$(VariationTitle)" Type="staged-build" BaseDir="$(ArtifactsPath)" Keys="Project=$(ProjectName)" Metadata="Stream=$(Branch);Platform=$(StagedPlatformFolder);BuildName=$(BuildNamePath);dashboard-category=Staged Builds" Files="#$(PlatformStageNodeName) Files"/>
+						</Node>
+					</Do>
 					
-					<Expand Name="WriteArtifactMetadataFile" Type="staged-build" Platform="$(StagedPlatformFolder)" Configuration="$(PackageConfigurations)" ArtifactDir="$(ArtifactsPath)"/>
-					<Tag Files="$(ArtifactsPath)/..." With="#$(PlatformStageNodeName) Files"/>
+					<Do If="$(PublishHordeArtifacts)">
+						<Node Name="$(PlatformStagedArtifactNodeName) Server" Requires="$(PlatformStageNodeName)">
+							<Property Name="ServerArtifactsPath" Value="$(ProjectOutputDirectory)/$(StagedPlatformFolder)" />
+							<Property Name="ServerArtifactsPath" Value="$(ProjectOutputDirectory)/$(StagedPlatformFolder)Server" If="$(ServerPackagedSeparately)" />
+							
+							<Expand Name="WriteArtifactMetadataFile" Type="staged-build" Platform="$(StagedPlatformFolder)$(VariationTitle)" Configuration="$(PackageConfigurations)" ArtifactDir="$(ServerArtifactsPath)"/>
+							<Tag Files="$(ServerArtifactsPath)/..." Except="$(TargetPlatform)Server/..." With="#$(PlatformStageNodeName) Server Files"/>
+
+							<CreateArtifact Name="$(ProjectName)-Server-Staged-$(StagedPlatformFolder)$(VariationTitle)" Type="staged-build" BaseDir="$(ServerArtifactsPath)" Keys="Project=$(ProjectName)" Metadata="Stream=$(Branch);Platform=$(StagedPlatformFolder);BuildName=$(BuildNamePath);dashboard-category=Staged Builds" Files="#$(PlatformStageNodeName) Server Files"/>
+						</Node>
+					</Do>
 
-					<CreateArtifact Name="$(ProjectName)-Staged-$(StagedPlatformFolder)" Type="staged-build" BaseDir="$(ArtifactsPath)" Keys="Project=$(ProjectName)" Metadata="Stream=$(Branch);Platform=$(StagedPlatformFolder);BuildName=$(BuildNamePath);dashboard-category=Staged Builds" Files="#$(PlatformStageNodeName) Files"/>
-				</Node>
-			</Do>
-			
-			<Do If="$(PublishHordeArtifacts)">
-				<Node Name="$(PlatformStagedArtifactNodeName) Server" Requires="$(PlatformStageNodeName)">
-					<Property Name="ServerArtifactsPath" Value="$(ProjectOutputDirectory)/$(StagedPlatformFolder)" />
-					<Property Name="ServerArtifactsPath" Value="$(ProjectOutputDirectory)/$(StagedPlatformFolder)Server" If="$(ServerPackagedSeparately)" />
+					<!-- Package -->
+					<Node Name="$(PlatformPackageNodeName)" Requires="$(PlatformPublishStageDependencyName)">
+						<ForEach Name="PackageConfiguration" Values="$(PackageConfigurations)" Separator="+">
+							<Property Name="ArchiveDirectory" Value="$(ProjectOutputDirectory)/Packages/$(StagedPlatformFolder)$(VariationTitle)/$(PackageConfiguration)" />
+							<Command Name="BuildCookRun" Arguments="$(BCRArgs) -configuration=$(PackageConfiguration) $(StageAndPackageArguments) -skipbuild -skipcook -skipstage -skippak -package $(PackageOnlyArguments) -archive -archivedirectory=&quot;$(ArchiveDirectory)&quot;" />
+						</ForEach>
+					</Node>
 					
-					<Expand Name="WriteArtifactMetadataFile" Type="staged-build" Platform="$(StagedPlatformFolder)" Configuration="$(PackageConfigurations)" ArtifactDir="$(ServerArtifactsPath)"/>
-					<Tag Files="$(ServerArtifactsPath)/..." Except="$(TargetPlatform)Server/..." With="#$(PlatformStageNodeName) Server Files"/>
-
-					<CreateArtifact Name="$(ProjectName)-Server-Staged-$(StagedPlatformFolder)" Type="staged-build" BaseDir="$(ServerArtifactsPath)" Keys="Project=$(ProjectName)" Metadata="Stream=$(Branch);Platform=$(StagedPlatformFolder);BuildName=$(BuildNamePath);dashboard-category=Staged Builds" Files="#$(PlatformStageNodeName) Server Files"/>
-				</Node>
-			</Do>
-
-			<!-- Package -->
-			<Node Name="$(PlatformPackageNodeName)" Requires="$(PlatformPublishStageDependencyName)">
-				<ForEach Name="PackageConfiguration" Values="$(PackageConfigurations)" Separator="+">
-					<Property Name="ArchiveDirectory" Value="$(ProjectOutputDirectory)/Packages/$(StagedPlatformFolder)/$(PackageConfiguration)" />
-					<Command Name="BuildCookRun" Arguments="$(BCRArgs) -configuration=$(PackageConfiguration) $(StageAndPackageArguments) -skipbuild -skipcook -skipstage -skippak -package $(PackageOnlyArguments) -archive -archivedirectory=&quot;$(ArchiveDirectory)&quot;" />
-				</ForEach>
-			</Node>
-			
-			<Node Name="$(PlatformPackageNodeName) Server" Requires="$(PlatformPublishStageDependencyName)">
-				<ForEach Name="PackageConfiguration" Values="$(PackageConfigurations)" Separator="+">
-					<Property Name="ServerArchiveDirectory" Value="$(ProjectOutputDirectory)/Packages/$(StagedPlatformFolder)/$(PackageConfiguration)" />
-					<Property Name="ServerArchiveDirectory" Value="$(ProjectOutputDirectory)/Packages/$(TargetPlatform)Server/$(PackageConfiguration)" If="$(ServerPackagedSeparately)" />
-					<Property Name="BCRArgsServer" Value="$(BCRArgs) -server" />	
-				<Command Name="BuildCookRun" Arguments="$(BCRArgsServer) -configuration=$(PackageConfiguration) $(StageAndPackageServerArguments) -skipbuild -skipcook -skipstage -skippak -package $(PackageOnlyArguments) -archive -archivedirectory=&quot;$(ServerArchiveDirectory)&quot;" />
-				</ForEach>
-			</Node>
-			
-			<!-- Publish (Packages) -->
-			<Node Name="$(PlatformPublishPackagedNodeName)" Requires="$(PlatformPackageNodeName)">
-				<Log Message="Network upload is disabled, skipping upload." If="!$(PublishToNetworkOutputDirectory)"/>
-				<Do If="$(PublishToNetworkOutputDirectory)">
-					<Error Message="NetworkOutputDirectory is not set" If="'$(NetworkOutputDirectory)' == ''"/>
-					<ForEach Name="PackageConfiguration" Values="$(PackageConfigurations)" Separator="+">
-						<Property Name="ArchiveDirectory" Value="$(NetworkOutputDirectory)/Packages/$(StagedPlatformFolder)"/>
-						<Command Name="BuildCookRun" Arguments="$(BCRArgs) -configuration=$(PackageConfiguration) $(StageAndPackageArguments) -skipbuild -skipcook -skipstage -skippak -skippackage -archive -archivedirectory=&quot;$(ArchiveDirectory)&quot;" />
-					</ForEach>
-				</Do>
-			</Node>
-			
-			<Node Name="$(PlatformPublishPackagedNodeName) Server" Requires="$(PlatformPackageNodeName)">
-				<Log Message="Network upload is disabled, skipping upload." If="!$(PublishToNetworkOutputDirectory)"/>
-				<Do If="$(PublishToNetworkOutputDirectory)">
-					<Error Message="NetworkOutputDirectory is not set" If="'$(NetworkOutputDirectory)' == ''"/>
-					<ForEach Name="PackageConfiguration" Values="$(PackageConfigurations)" Separator="+">
-						<Property Name="ServerArchiveDirectory" Value="$(NetworkOutputDirectory)/Packages/$(StagedPlatformFolder)"/>
-						<Property Name="ServerArchiveDirectory" Value="$(NetworkOutputDirectory)/Packages/$(TargetPlatform)Server" If="$(ServerPackagedSeparately)"/>
-						<Property Name="BCRArgsServer" Value="$(BCRArgs) -server" />	
-						<Command Name="BuildCookRun" Arguments="$(BCRArgsServer) -configuration=$(PackageConfiguration) $(StageAndPackageServerArguments) -skipbuild -skipcook -skipstage -skippak -skippackage -archive -archivedirectory=&quot;$(ServerArchiveDirectory)&quot;" />
-					</ForEach>
-				</Do>
-			</Node>
-
-			<!-- Publish (Packages) (Cloud) -->
-			<Do  If="$(PublishCloudArtifacts)">
-				<Node Name="$(CloudArtifactPlatformPackagedNodeName)" Requires="$(PlatformPackageNodeName)">
-					<ForEach Name="PackageConfiguration" Values="$(PackageConfigurations)" Separator="+">
-						<Property Name="TargetDirectory" Value="$(ProjectOutputDirectory)/Packages/$(StagedPlatformFolder)/$(PackageConfiguration)"/>						
-						<Expand Name="WriteArtifactMetadataFile" Type="packaged-build" Platform="$(StagedPlatformFolder)" Configuration="$(PackageConfiguration)" ArtifactDir="$(TargetDirectory)" />
-						<Tag Files="$(TargetDirectory)/..." With="#$(PlatformPackageNodeName) $(PackageConfiguration) Files"/>
-
-						<CreateCloudArtifact 
-						Name="$(ProjectName)-$(TargetPlatform)-Packaged-$(PackageConfiguration)" 
-						BuildGroup="$(BuildVersion)"
-						Type="packaged-build" 
-						BaseDir="$(TargetDirectory)" 
-						Metadata="cookPlatform=$(StagedPlatformFolder);Configuration=$(PackageConfiguration);dashboard-category=Packaged Builds - Cloud;$(SharedCloudBuildMetadata)" 
-						Project="$(TargetProject)" 
-						Platform="$(StagedPlatformFolder)" 
-						Files="#$(PlatformPackageNodeName) $(PackageConfiguration) Files"/>
-						<Delete Files="$(TargetDirectory)/artifactmetadata.json"/>
-					</ForEach>
-				</Node>
-			</Do>
-			
-			<!-- Publish (Packages) (Horde) -->
-			<Do  If="$(PublishHordeArtifacts)">
-				<Property Name="HordePackageSkipBCRArguments" Value="-stage -package"/>
-				<Property Name="HordePackageSkipBCRArguments" Value="-skipstage -skippackage" If="'$(TargetPlatform)' == 'Mac'"/>
-
-				<Node Name="$(HordeArtifactPlatformPackagedNodeName)" Requires="$(PlatformPackageNodeName)">
-					<ForEach Name="PackageConfiguration" Values="$(PackageConfigurations)" Separator="+">
-						<Property Name="ArchivePath" Value="$(ProjectOutputDirectory)/Packages/$(StagedPlatformFolder)/$(PackageConfiguration)" />
-						<Command Name="BuildCookRun" Arguments="$(BCRArgs) -configuration=$(PackageConfiguration) $(StageAndPackageArguments) -skipbuild -skipcook -skippak $(HordePackageSkipBCRArguments) -archive -archivedirectory=&quot;$(ArchivePath)&quot;" />
-						<Expand Name="WriteArtifactMetadataFile" Type="packaged-build" Platform="$(StagedPlatformFolder)" Configuration="$(PackageConfiguration)" ArtifactDir="$(ArchivePath)"/>
-						<Tag Files="$(ArchivePath)/..." Except="$(TargetPlatform)Server/...;*$(TargetPlatform)Server*" With="#$(PlatformPackageNodeName) $(PackageConfiguration) Files"/>
-						<CreateArtifact Name="$(ProjectName)-Packaged-$(StagedPlatformFolder)-$(PackageConfiguration)" Type="packaged-build" BaseDir="$(ArchivePath)" Keys="Project=$(ProjectName)" Metadata="Stream=$(Branch);Platform=$(StagedPlatformFolder);Configuration=$(PackageConfiguration);BuildName=$(BuildNamePath);dashboard-category=$(PackageConfiguration) Builds" Files="#$(PlatformPackageNodeName) $(PackageConfiguration) Files"/>
-					</ForEach>
-				</Node>
-			
-				<Node Name="$(HordeArtifactPlatformPackagedNodeName) Server" Requires="$(PlatformPackageNodeName)">
-					<ForEach Name="PackageConfiguration" Values="$(PackageConfigurations)" Separator="+">
-						<Property Name="ServerArchivePath" Value="$(ProjectOutputDirectory)/Packages/$(StagedPlatformFolder)/$(PackageConfiguration)" />
-						<Property Name="ServerArchivePath" Value="$(ProjectOutputDirectory)/Packages/$(StagedPlatformFolder)Server/$(PackageConfiguration)" If="$(ServerPackagedSeparately)" />
-						<Property Name="BCRArgsServer" Value="$(BCRArgs) -server" />
-						<Command Name="BuildCookRun" Arguments="$(BCRArgsServer) -configuration=$(PackageConfiguration) $(StageAndPackageServerArguments) -skipbuild -skipcook -skippak $(HordePackageSkipBCRArguments) -archive -archivedirectory=&quot;$(ServerArchivePath)&quot;" />
-						<Tag Files="$(ServerArchivePath)/..."  Except="$(TargetPlatform)Server/..." With="#$(PlatformPackageNodeName) $(PackageConfiguration) Server Files"/>
-						<CreateArtifact Name="$(ProjectName)-Server-Packaged-$(TargetPlatform)-$(PackageConfiguration)" Type="packaged-build" BaseDir="$(ServerArchivePath)" Keys="Project=$(ProjectName)" Metadata="Stream=$(Branch);Platform=$(TargetPlatform);Configuration=$(PackageConfiguration);BuildName=$(BuildNamePath);dashboard-category=$(PackageConfiguration) Builds" Files="#$(PlatformPackageNodeName) $(PackageConfiguration) Server Files"/>
-					</ForEach>
-				</Node>
-			</Do>
-
-			<!-- Check Build Size -->
-			<Node Name="$(PlatformCheckBuildSizeNodeName)" Requires="$(PlatformPackageNodeName)">
-				<ForEach Name="PackageConfiguration" Values="$(PackageConfigurations)" Separator="+">
-					<Command Name="CheckBuildSize" Arguments="-project=$(TargetName) -platform=$(TargetPlatform) -packagesize=$(BuildSizeLimit) -build=&quot;$(ProjectOutputDirectory)/Packages/$(StagedPlatformFolder)/$(PackageConfiguration)&quot; -configuration=$(PackageConfiguration)" />
-				</ForEach>
-			</Node>
-			
-		</Agent>
-		
-		<!-- Test Time! -->
-		<Property Name="AllPlatformInstallNodes" Value="" />
-		<Property Name="AllPlatformBootTestNodes" Value="" />
-		<Property Name="AllPlatformTestNodes" Value="" />
-
-		<Property Name="UsePackaged" Value="(ContainsItem('$(PlatformsThatRequirePackaging)', '$(TargetPlatform)', '+') Or !ContainsItem('$(PlatformsToSkipPackaging)', '$(TargetPlatform)', '+'))" />
-
-		<!-- Create the -build argument for Gauntlet based on whether builds are local or on the network-->
-		<Property Name="GauntletBuildArg" Value="-build=&quot;$(ProjectOutputDirectory)/$(StagedPlatformFolder)&quot;" />
-		<Property Name="GauntletBuildArg" Value="-build=&quot;$(ProjectOutputDirectory)/Packages/$(StagedPlatformFolder)&quot; -packaged" If="$(UsePackaged)" />
-
-		<Do If="!$(Horde) And ('$(NetworkOutputDirectory)' != '' Or '$(GeneratedBuildVersion)' != '')">
-			<Property Name="TestDirectory" Value="$(NetworkOutputDirectory)" />
-			<Property Name="TestDirectory" Value="$(GeneratedBuildDir)" If="'$(GeneratedBuildVersion)' != ''"/>
-			<Property Name="GauntletBuildArg" Value="-build=&quot;$(TestDirectory)&quot;" />
-			<Property Name="GauntletBuildArg" Value="$(GauntletBuildArg) -packaged" If="$(UsePackaged)" />
-			<Property Name="GauntletBuildArg" Value="$(GauntletBuildArg) -destlocalinstalldir=&quot;$(LocalInstallDirectory)&quot;" If="!$(IsLocalPlatform) And $(WithDeviceReservation)"/>
-		</Do>
-
-		<!-- Test Configs -->
-		<ForEach Name="TestConfig" Values="$(TestConfigList)" Separator="^">
-			<Property Name="TestConfigTitle" Value=""/>
-			<Regex Input="$(TestConfig)" Pattern="(.*)\|(.*)\|(.*)\|(.*)\|(.*)\|(.*)" Capture="TestConfigName;TestConfigTargetPlatform;ConfigArtifactPostfix;ConfigAdditionalTestArguments;ConfigAdditionalAppArguments;ConfigAgentType"/>
-
-			<!--If this is the default test config i.e. "All" or Empty ensure we create nodes for it-->
-			<Property Name="TestConfigTargetPlatform" Value="$(TargetPlatform)" If="'$(TestConfigTargetPlatform)' == 'All' Or '$(TestConfigTargetPlatform)' == ''"/>
-
-			<Do If="'$(TestConfigTargetPlatform)' == '$(TargetPlatform)'">
-
-				<Property Name="TestConfigTitle" Value=" $(TestConfigName)" If="'$(TestConfigName)' != ''"/>
-
-				<Property Name="TestConfigArgs" Value="" />
-				<Property Name="TestConfigArgs" Value="$(ConfigAdditionalTestArguments)" If="'$(ConfigAdditionalTestArguments)' != ''"/>
-				<Property Name="TestConfigArgs" Value="$(TestConfigArgs) -ArtifactPostfix=$(ConfigArtifactPostfix)" If="'$(ConfigArtifactPostfix)' != ''"  />
-				<Property Name="TestConfigArgs" Value="$(TestConfigArgs) -args=&quot;$(ConfigAdditionalAppArguments)&quot;" If="'$(ConfigAdditionalAppArguments)' != ''" />
-
-				<Property Name="TestTargetAgentType" Value="$(TestAgentType)"/>
-				<Property Name="TestTargetAgentType" Value="$(ConfigAgentType)" If="'$(ConfigAgentType)' != ''"/>
-
-				<Agent Name="Test$(TestConfigTitle) $(PlatformAgentName)" Type="$(TestTargetAgentType)">
-					<!-- Run through all configs -->
-					<ForEach Name="TestConfiguration" Values="$(TestConfigurations)" Separator="+">
-						<Property Name="NodeAnnotations" Value=""/>
-						<!-- Prerequisites for the test to run -->
-						<Property Name="TestPrerequisites" Value="$(RequirePlatformStaged)" />
-						<Property Name="TestPrerequisites" Value="$(RequirePlatformPackaged)" If="$(UsePackaged)"/>
-						<Property Name="TestPrerequisites" Value="$(PlatformToolsNodeName)" If="'$(GeneratedBuildVersion)' != ''"/>
-						<Property Name="TestPrerequisites" Value="$(PlatformToolsNodeName);$(PlatformEditorNodeName)" If="'$(GeneratedBuildVersion)' != '' and $(TestsRequireEditor)"/>
-
-						<Do If="$(PublishHordeArtifacts) and !$(RetrieveUsingCloudArtifacts)">
-							<Property Name="HordeArtifactPlatformRetrieveWithConfig" Value="$(HordeArtifactPlatformRetrieve) $(TestConfiguration)$(TestConfigTitle) (Horde)"/>
-							<Node Name="$(HordeArtifactPlatformRetrieveWithConfig)" Requires="$(TestPrerequisites)">
-								<Switch>
-									<Case If="$(UsePackaged)">
-										<RetrieveArtifact Name="$(ProjectName)-Packaged-$(StagedPlatformFolder)-$(TestConfiguration)" Type="packaged-build" OutputDir="$(ProjectOutputDirectory)/Packages/$(StagedPlatformFolder)/$(TestConfiguration)" />
-									</Case>
-									<Default>
-										<RetrieveArtifact Name="$(ProjectName)-Staged-$(StagedPlatformFolder)" Type="staged-build" OutputDir="$(ProjectOutputDirectory)/$(StagedPlatformFolder)" />
-									</Default>
-								</Switch>
-							</Node>
-							<Property Name="TestPrerequisites" Value="$(TestPrerequisites);$(HordeArtifactPlatformRetrieveWithConfig)" />
-						</Do>
-
-						<Do If="$(RetrieveUsingCloudArtifacts)">
-							<Property Name="CloudArtifactPlatformRetrieveWithConfig" Value="$(CloudArtifactPlatformRetrieve) $(TestConfiguration)$(TestConfigTitle) (Cloud)"/>
-							<Node Name="$(CloudArtifactPlatformRetrieveWithConfig)" Requires="$(TestPrerequisites)">
-								<Switch>
-									<Case If="$(UsePackaged)">
-										<RetrieveCloudArtifact 
-										Name="$(ProjectName)-$(TargetPlatform)-Packaged-$(TestConfiguration)" 
-										Type="packaged-build" 
-										OutputDir="$(ProjectOutputDirectory)/Packages/$(StagedPlatformFolder)/$(TestConfiguration)" 
-										Project="$(TargetProject)" 
-										Platform="$(StagedPlatformFolder)"
-										BuildGroup="$(BuildVersion)" 
-										AllowMultipleMatches="true"
-										/>
-									</Case>
-									<Default>
-										<RetrieveCloudArtifact 
-										Name="$(ProjectName)-$(TargetPlatform)-Staged" 
-										Type="staged-build" 
-										OutputDir="$(ProjectOutputDirectory)/$(StagedPlatformFolder)" 
-										Project="$(TargetProject)" 
-										Platform="$(StagedPlatformFolder)" 
-										BuildGroup="$(BuildVersion)" 
-										AllowMultipleMatches="true"
-										/>
-									</Default>
-								</Switch>
-							</Node>
-
-							<Property Name="TestPrerequisites" Value="$(TestPrerequisites);$(CloudArtifactPlatformRetrieveWithConfig)" />
+					<Node Name="$(PlatformPackageNodeName) Server" Requires="$(PlatformPublishStageDependencyName)">
+						<ForEach Name="PackageConfiguration" Values="$(PackageConfigurations)" Separator="+">
+							<Property Name="ServerArchiveDirectory" Value="$(ProjectOutputDirectory)/Packages/$(StagedPlatformFolder)$(VariationTitle)/$(PackageConfiguration)" />
+							<Property Name="ServerArchiveDirectory" Value="$(ProjectOutputDirectory)/Packages/$(TargetPlatform)Server$(VariationTitle)/$(PackageConfiguration)" If="$(ServerPackagedSeparately)" />
+							<Property Name="BCRArgsServer" Value="$(BCRArgs) -server" />	
+						<Command Name="BuildCookRun" Arguments="$(BCRArgsServer) -configuration=$(PackageConfiguration) $(StageAndPackageServerArguments) -skipbuild -skipcook -skipstage -skippak -package $(PackageOnlyArguments) -archive -archivedirectory=&quot;$(ServerArchiveDirectory)&quot;" />
+						</ForEach>
+					</Node>
+					
+					<!-- Publish (Packages) -->
+					<Node Name="$(PlatformPublishPackagedNodeName)" Requires="$(PlatformPackageNodeName)">
+						<Log Message="Network upload is disabled, skipping upload." If="!$(PublishToNetworkOutputDirectory)"/>
+						<Do If="$(PublishToNetworkOutputDirectory)">
+							<Error Message="NetworkOutputDirectory is not set" If="'$(NetworkOutputDirectory)' == ''"/>
+							<ForEach Name="PackageConfiguration" Values="$(PackageConfigurations)" Separator="+">
+								<Property Name="ArchiveDirectory" Value="$(NetworkOutputDirectory)/Packages/$(StagedPlatformFolder)$(VariationTitle)"/>
+								<Command Name="BuildCookRun" Arguments="$(BCRArgs) -configuration=$(PackageConfiguration) $(StageAndPackageArguments) -skipbuild -skipcook -skipstage -skippak -skippackage -archive -archivedirectory=&quot;$(ArchiveDirectory)&quot;" />
+							</ForEach>
 						</Do>
-
-						<Property Name="GauntletTestArgs" Value="$(GauntletTargetTestArgs)"/>
-						<Property Name="GauntletTestArgs" Value="$(GauntletTargetTestAltArgs)" If="ContainsItem('$(AlternateDevicePoolPlatforms)', '$(TargetPlatform)', ';')"/>
-
-						<Property Name="GauntletTestArgs" Value="$(GauntletTestArgs) -adbkeys=&quot;$(AdbKeysPath)&quot;" If="'$(AdbKeysPath)' != '' and '$(TargetPlatform)' == 'Android'"/>
-
-						<!-- Install step -->
-						<Do If="!$(IsLocalPlatform) And $(WithDeviceReservation)">
-							<Property Name="NodeAnnotations" Value=""/>
-							<Property Name="NodeAnnotations" Value="Workflow=$(AnnotationsTarget)" If="ContainsItem('$(AnnotationsPlatforms)', '$(TargetPlatform)', ';') and ContainsItem('$(AnnotationsTests)', 'Install', ';')"/>
-							<Property Name="InstallNodeName" Value="$(PreNodeName)$(TargetPlatform) $(TestConfiguration)$(TestConfigTitle) Install Game"/>
-							<Property Name="InstallArgs" Value="-ClientCount=1 $(GauntletPreInstallArgs)"/>
-							<Do If="'$(AutomationReportOutputDirectory)' != ''">
-								<Property Name="OutputFragmentPath" Value="$(TargetPlatform)/$(TestConfiguration)/UE.InstallOnly" />
-								<Property Name="OutputReportPath" Value="$(AutomationReportOutputDirectory)/$(OutputFragmentPath)" />
-								<Property Name="GauntletReportArgs" Value="-logdir=&quot;$(OutputReportPath)&quot; -artifactname=Gauntlet" />
-							</Do>
-							<Node Name="$(InstallNodeName)" Requires="$(TestPrerequisites)" Annotations="DeviceReserve=Begin;$(NodeAnnotations)">
-								<Command Name="RunUnreal" Arguments="-project=$(TargetProject) -test=&quot;UE.InstallOnly&quot; -platform=$(TargetPlatform) -configuration=$(TestConfiguration) -target=$(TargetName) $(GauntletBuildArg) $(GauntletTestArgs) $(GauntletReportArgs) $(InstallArgs)" />
-							</Node>
-							<Property Name="TestPrerequisites" Value="$(TestPrerequisites);$(InstallNodeName)" />
-							<Do If="!ContainsItem('$(TargetPlatformsToSkipBootTests)', '$(TargetPlatform)', '+') And !ContainsItem('$(TargetPlatformsToSkipTests)', '$(TargetPlatform)', '+')">
-								<Property Name="AllPlatformInstallNodes" Value="$(AllPlatformInstallNodes);$(InstallNodeName)"/>
-							</Do>
-							<!-- Skip install for all other test node -->
-							<Property Name="GauntletTestArgs" Value="$(GauntletTestArgs) $(GauntletTestSkipInstallArgs)"/>
+					</Node>
+					
+					<Node Name="$(PlatformPublishPackagedNodeName) Server" Requires="$(PlatformPackageNodeName)">
+						<Log Message="Network upload is disabled, skipping upload." If="!$(PublishToNetworkOutputDirectory)"/>
+						<Do If="$(PublishToNetworkOutputDirectory)">
+							<Error Message="NetworkOutputDirectory is not set" If="'$(NetworkOutputDirectory)' == ''"/>
+							<ForEach Name="PackageConfiguration" Values="$(PackageConfigurations)" Separator="+">
+								<Property Name="ServerArchiveDirectory" Value="$(NetworkOutputDirectory)/Packages/$(StagedPlatformFolder)$(VariationTitle)"/>
+								<Property Name="ServerArchiveDirectory" Value="$(NetworkOutputDirectory)/Packages/$(TargetPlatform)Server$(VariationTitle)" If="$(ServerPackagedSeparately)"/>
+								<Property Name="BCRArgsServer" Value="$(BCRArgs) -server" />	
+								<Command Name="BuildCookRun" Arguments="$(BCRArgsServer) -configuration=$(PackageConfiguration) $(StageAndPackageServerArguments) -skipbuild -skipcook -skipstage -skippak -skippackage -archive -archivedirectory=&quot;$(ServerArchiveDirectory)&quot;" />
+							</ForEach>
 						</Do>
-
-						<!-- Expand the tests here -->
-						<Expand Name="PreTargetTests"/>
-
-						<!-- Declare a boot test for the target-->
-						<Property Name="NodeAnnotations" Value=""/>
-						<Property Name="NodeAnnotations" Value="Workflow=$(AnnotationsTarget)" If="ContainsItem('$(AnnotationsPlatforms)', '$(TargetPlatform)', ';') and ContainsItem('$(AnnotationsTests)', 'BootTest', ';')"/>
-						<Property Name="BootTestName" Value="$(PreNodeName)$(TargetPlatform) $(TestConfiguration)$(TestConfigTitle) BootTest" />
-						<Node Name="$(BootTestName)" Requires="$(TestPrerequisites)" Annotations="$(NodeAnnotations)">
-							<Do If="'$(AutomationReportOutputDirectory)' != ''">
-								<Property Name="OutputFragmentPath" Value="$(TargetPlatform)/$(TestConfiguration)/UE.BootTest" />
-								<Property Name="OutputReportPath" Value="$(AutomationReportOutputDirectory)/$(OutputFragmentPath)" />
-								<Property Name="GauntletReportArgs" Value="-logdir=&quot;$(OutputReportPath)&quot; -artifactname=Gauntlet" />
-							</Do>
-							<!-- Run the boot test -->
-							<Command Name="RunUnreal" Arguments="-project=$(TargetProject) -test=&quot;UE.BootTest&quot; -platform=$(TargetPlatform) -configuration=$(TestConfiguration) -target=$(TargetName) $(GauntletBuildArg) $(GauntletTestArgs) $(GauntletReportArgs) $(TestConfigArgs)" />
+					</Node>
+
+					<!-- Publish (Packages) (Cloud) -->
+					<Do  If="$(PublishCloudArtifacts)">
+						<Node Name="$(CloudArtifactPlatformPackagedNodeName)" Requires="$(PlatformPackageNodeName)">
+							<ForEach Name="PackageConfiguration" Values="$(PackageConfigurations)" Separator="+">
+								<Property Name="TargetDirectory" Value="$(ProjectOutputDirectory)/Packages/$(StagedPlatformFolder)/$(PackageConfiguration)"/>						
+								<Expand Name="WriteArtifactMetadataFile" Type="packaged-build" Platform="$(StagedPlatformFolder)$(VariationTitle)" Configuration="$(PackageConfiguration)" ArtifactDir="$(TargetDirectory)" />
+								<Tag Files="$(TargetDirectory)/..." With="#$(PlatformPackageNodeName) $(PackageConfiguration) Files"/>
+
+								<CreateCloudArtifact 
+								Name="$(ProjectName)-$(TargetPlatform)$(VariationTitle)-Packaged-$(PackageConfiguration)" 
+								BuildGroup="$(BuildVersion)"
+								Type="packaged-build" 
+								BaseDir="$(TargetDirectory)" 
+								Metadata="cookPlatform=$(StagedPlatformFolder)$(VariationTitle);Configuration=$(PackageConfiguration);dashboard-category=Packaged Builds - Cloud;$(SharedCloudBuildMetadata)" 
+								Project="$(TargetProject)" 
+								Platform="$(StagedPlatformFolder)" 
+								Files="#$(PlatformPackageNodeName) $(PackageConfiguration) Files"/>
+								<Delete Files="$(TargetDirectory)/artifactmetadata.json"/>
+							</ForEach>
 						</Node>
-
-						<!-- Link the boot test in to the prerequisites and generated nodes if its enabled-->
-						<Do If="$(TargetBootTest) and !ContainsItem('$(TargetPlatformsToSkipBootTests)', '$(TargetPlatform)', '+')">
-							<Property Name="TestPrerequisites" Value="$(TestPrerequisites);$(BootTestName)" />
-							<Property Name="AllPlatformBootTestNodes" Value="$(AllPlatformBootTestNodes);$(BootTestName)"/>
-						</Do>
-
-						<!-- Go through the list of tests -->
-						<ForEach Name="TestName" Values="$(TargetTestList)" Separator="+">
-							<Property Name="TestNameTitle" Value="$(TestName)"/>
-							<Expand Name="SanitizeTestNameTitle"/>
-							<!-- Run and report on the test-->
-							<Do If="'$(AutomationReportOutputDirectory)' != ''">
-								<Property Name="OutputFragmentPath" Value="$(TargetPlatform)/$(TestConfiguration)/$(TestNameTitle)" />
-								<Property Name="OutputReportPath" Value="$(AutomationReportOutputDirectory)/$(OutputFragmentPath)" />
-								<Property Name="GauntletReportArgs" Value="-logdir=&quot;$(OutputReportPath)&quot; -artifactname=Gauntlet -ReportExportPath=&quot;$(OutputReportPath)&quot;" />
-								<Property Name="GauntletReportArgs" Value="$(GauntletReportArgs) -ReportURL=&quot;$(AutomationReportUrl)/$(OutputFragmentPath)/&quot;" If="'$(AutomationReportUrl)' != ''" />
-							</Do>
-							<Property Name="TestNodeName" Value="$(PreNodeName)$(TargetPlatform) $(TestConfiguration)$(TestConfigTitle) $(TestNameTitle)" />
-							<Property Name="NodeAnnotations" Value=""/>
-							<Property Name="NodeAnnotations" Value="Workflow=$(AnnotationsTarget)" If="ContainsItem('$(AnnotationsPlatforms)', '$(TargetPlatform)', ';') and ContainsItem('$(AnnotationsTests)', '$(TestName)', ';')"/>
-
-							<Property Name="TestNameBeforeParentheses" Value="$(TestName)"/>
-							<Property Name="TestArgumentsWithinParentheses" Value=""/>
-							<Expand Name="SplitCombinedTestNameAndArgs"/>
-
-							<!-- Test 'Requires' should be the package game and the publish folder is used; however some tests may require the editor ie: UE.TargetAutomation -->
-							<Property Name="AdditionalTestPrerequisites" Value=""/>
-							<ForEach Name="ToMatchName" Values="$(RequiringEditorTargetTestList)" Separator="+">
-								<Do If="Contains('$(TestNameBeforeParentheses)', '$(ToMatchName)') And !ContainsItem('$(TestPrerequisites)', '$(PlatformEditorNodeName)', ';')">
-									<Property Name="AdditionalTestPrerequisites" Value="$(AdditionalTestPrerequisites);$(PlatformEditorNodeName)" />
-								</Do>
+					</Do>
+					
+					<!-- Publish (Packages) (Horde) -->
+					<Do  If="$(PublishHordeArtifacts)">
+						<Property Name="HordePackageSkipBCRArguments" Value="-stage -package"/>
+						<Property Name="HordePackageSkipBCRArguments" Value="-skipstage -skippackage" If="'$(TargetPlatform)' == 'Mac'"/>
+
+						<Node Name="$(HordeArtifactPlatformPackagedNodeName)" Requires="$(PlatformPackageNodeName)">
+							<ForEach Name="PackageConfiguration" Values="$(PackageConfigurations)" Separator="+">
+								<Property Name="ArchivePath" Value="$(ProjectOutputDirectory)/Packages/$(StagedPlatformFolder)$(VariationTitle)/$(PackageConfiguration)" />
+								<Command Name="BuildCookRun" Arguments="$(BCRArgs) -configuration=$(PackageConfiguration) $(StageAndPackageArguments) -skipbuild -skipcook -skippak $(HordePackageSkipBCRArguments) -archive -archivedirectory=&quot;$(ArchivePath)&quot;" />
+								<Expand Name="WriteArtifactMetadataFile" Type="packaged-build" Platform="$(StagedPlatformFolder)$(VariationTitle)" Configuration="$(PackageConfiguration)" ArtifactDir="$(ArchivePath)"/>
+								<Tag Files="$(ArchivePath)/..." Except="$(TargetPlatform)Server/...;*$(TargetPlatform)Server*" With="#$(PlatformPackageNodeName) $(PackageConfiguration) Files"/>
+								<CreateArtifact Name="$(ProjectName)-Packaged-$(StagedPlatformFolder)$(VariationTitle)-$(PackageConfiguration)" Type="packaged-build" BaseDir="$(ArchivePath)" Keys="Project=$(ProjectName)" Metadata="Stream=$(Branch);Platform=$(StagedPlatformFolder);Configuration=$(PackageConfiguration);BuildName=$(BuildNamePath);dashboard-category=$(PackageConfiguration) Builds" Files="#$(PlatformPackageNodeName) $(PackageConfiguration) Files"/>
 							</ForEach>
+						</Node>
+					
+						<Node Name="$(HordeArtifactPlatformPackagedNodeName) Server" Requires="$(PlatformPackageNodeName)">
+							<ForEach Name="PackageConfiguration" Values="$(PackageConfigurations)" Separator="+">
+								<Property Name="ServerArchivePath" Value="$(ProjectOutputDirectory)/Packages/$(StagedPlatformFolder)$(VariationTitle)/$(PackageConfiguration)" />
+								<Property Name="ServerArchivePath" Value="$(ProjectOutputDirectory)/Packages/$(StagedPlatformFolder)Server$(VariationTitle)/$(PackageConfiguration)" If="$(ServerPackagedSeparately)" />
+								<Property Name="BCRArgsServer" Value="$(BCRArgs) -server" />
+								<Command Name="BuildCookRun" Arguments="$(BCRArgsServer) -configuration=$(PackageConfiguration) $(StageAndPackageServerArguments) -skipbuild -skipcook -skippak $(HordePackageSkipBCRArguments) -archive -archivedirectory=&quot;$(ServerArchivePath)&quot;" />
+								<Tag Files="$(ServerArchivePath)/..."  Except="$(TargetPlatform)Server/..." With="#$(PlatformPackageNodeName) $(PackageConfiguration) Server Files"/>
+								<CreateArtifact Name="$(ProjectName)-Server-Packaged-$(TargetPlatform)$(VariationTitle)-$(PackageConfiguration)" Type="packaged-build" BaseDir="$(ServerArchivePath)" Keys="Project=$(ProjectName)" Metadata="Stream=$(Branch);Platform=$(TargetPlatform);Configuration=$(PackageConfiguration);BuildName=$(BuildNamePath);dashboard-category=$(PackageConfiguration) Builds" Files="#$(PlatformPackageNodeName) $(PackageConfiguration) Server Files"/>
+							</ForEach>
+						</Node>
+					</Do>
 
-							<!-- 'Requires' should be the package game and the publish folder is used; however some tests may require the editor ie: UE.TargetAutomation -->
-							<Node Name="$(TestNodeName)" Requires="$(TestPrerequisites);$(AdditionalTestPrerequisites)" Annotations="$(NodeAnnotations)">
-								<Command Name="RunUnreal" Arguments="-project=$(TargetProject) -test=&quot;$(TestNameBeforeParentheses)&quot; -platform=$(TargetPlatform) -configuration=$(TestConfiguration) -target=$(TargetName) $(GauntletBuildArg) $(GauntletTestArgs) $(TestArgumentsWithinParentheses) $(GauntletReportArgs) $(TestConfigArgs)"/>
-							</Node>
-							<!-- Add it to our list -->
-							<Do If="!ContainsItem('$(TargetPlatformsToSkipTests)', '$(TargetPlatform)', '+') And (!Contains('$(TestListToIncludePlatforms)', '^$(TestNameBeforeParentheses)|') Or ContainsItem('$(TestListToIncludePlatforms)', '$(TestNameBeforeParentheses)|$(TargetPlatform)', '^'))">
-								<Property Name="AllPlatformTestNodes" Value="$(AllPlatformTestNodes);$(TestNodeName)"/>
-							</Do>
+					<!-- Check Build Size -->
+					<Node Name="$(PlatformCheckBuildSizeNodeName)" Requires="$(PlatformPackageNodeName)">
+						<ForEach Name="PackageConfiguration" Values="$(PackageConfigurations)" Separator="+">
+							<Command Name="CheckBuildSize" Arguments="-project=$(TargetName) -platform=$(TargetPlatform) -packagesize=$(BuildSizeLimit) -build=&quot;$(ProjectOutputDirectory)/Packages/$(StagedPlatformFolder)/$(PackageConfiguration)&quot; -configuration=$(PackageConfiguration)" />
 						</ForEach>
+					</Node>
+					
+				</Agent>
+				
+				<!-- Test Time! -->
+				<Property Name="AllPlatformInstallNodes" Value="" />
+				<Property Name="AllPlatformBootTestNodes" Value="" />
+				<Property Name="AllPlatformTestNodes" Value="" />
+				<Do If="!ContainsItem('$(VariationsToSkipTests)', '$(Variation)', '+')">
+
+					<Property Name="UsePackaged" Value="(ContainsItem('$(PlatformsThatRequirePackaging)', '$(TargetPlatform)', '+') Or !ContainsItem('$(PlatformsToSkipPackaging)', '$(TargetPlatform)', '+'))" />
+
+					<!-- Create the -build argument for Gauntlet based on whether builds are local or on the network-->
+					<Property Name="GauntletBuildArg" Value="-build=&quot;$(ProjectOutputDirectory)/$(StagedPlatformFolder)&quot;" />
+					<Property Name="GauntletBuildArg" Value="-build=&quot;$(ProjectOutputDirectory)/Packages/$(StagedPlatformFolder)&quot; -packaged" If="$(UsePackaged)" />
+
+					<Do If="!$(Horde) And ('$(NetworkOutputDirectory)' != '' Or '$(GeneratedBuildVersion)' != '')">
+						<Property Name="TestDirectory" Value="$(NetworkOutputDirectory)" />
+						<Property Name="TestDirectory" Value="$(GeneratedBuildDir)" If="'$(GeneratedBuildVersion)' != ''"/>
+						<Property Name="GauntletBuildArg" Value="-build=&quot;$(TestDirectory)&quot;" />
+						<Property Name="GauntletBuildArg" Value="$(GauntletBuildArg) -packaged" If="$(UsePackaged)" />
+						<Property Name="GauntletBuildArg" Value="$(GauntletBuildArg) -destlocalinstalldir=&quot;$(LocalInstallDirectory)&quot;" If="!$(IsLocalPlatform) And $(WithDeviceReservation)"/>
+					</Do>
 
-						<!-- Expand the tests here -->
-						<Expand Name="AdditionalTestNodes"/>
+					<!-- Test Configs -->
+					<ForEach Name="TestConfig" Values="$(TestConfigList)" Separator="^">
+						<Property Name="TestConfigTitle" Value=""/>
+						<Regex Input="$(TestConfig)" Pattern="(.*)\|(.*)\|(.*)\|(.*)\|(.*)\|(.*)" Capture="TestConfigName;TestConfigTargetPlatform;ConfigArtifactPostfix;ConfigAdditionalTestArguments;ConfigAdditionalAppArguments;ConfigAgentType"/>
+
+						<!--If this is the default test config i.e. "All" or Empty ensure we create nodes for it-->
+						<Property Name="TestConfigTargetPlatform" Value="$(TargetPlatform)" If="'$(TestConfigTargetPlatform)' == 'All' Or '$(TestConfigTargetPlatform)' == ''"/>
+
+						<Do If="'$(TestConfigTargetPlatform)' == '$(TargetPlatform)'">
+
+							<Property Name="TestConfigTitle" Value=" $(TestConfigName)" If="'$(TestConfigName)' != ''"/>
+
+							<Property Name="TestConfigArgs" Value="" />
+							<Property Name="TestConfigArgs" Value="$(ConfigAdditionalTestArguments)" If="'$(ConfigAdditionalTestArguments)' != ''"/>
+							<Property Name="TestConfigArgs" Value="$(TestConfigArgs) -ArtifactPostfix=$(ConfigArtifactPostfix)" If="'$(ConfigArtifactPostfix)' != ''"  />
+							<Property Name="TestConfigArgs" Value="$(TestConfigArgs) -args=&quot;$(ConfigAdditionalAppArguments)&quot;" If="'$(ConfigAdditionalAppArguments)' != ''" />
+
+							<Property Name="TestTargetAgentType" Value="$(TestAgentType)"/>
+							<Property Name="TestTargetAgentType" Value="$(ConfigAgentType)" If="'$(ConfigAgentType)' != ''"/>
+
+							<Agent Name="Test$(TestConfigTitle) $(PlatformAgentName)" Type="$(TestTargetAgentType)">
+								<!-- Run through all configs -->
+								<ForEach Name="TestConfiguration" Values="$(TestConfigurations)" Separator="+">
+									<Property Name="NodeAnnotations" Value=""/>
+									<!-- Prerequisites for the test to run -->
+									<Property Name="TestPrerequisites" Value="$(RequirePlatformStaged)" />
+									<Property Name="TestPrerequisites" Value="$(RequirePlatformPackaged)" If="$(UsePackaged)"/>
+									<Property Name="TestPrerequisites" Value="$(PlatformToolsNodeName)" If="'$(GeneratedBuildVersion)' != ''"/>
+									<Property Name="TestPrerequisites" Value="$(PlatformToolsNodeName);$(PlatformEditorNodeName)" If="'$(GeneratedBuildVersion)' != '' and $(TestsRequireEditor)"/>
+
+									<Do If="$(PublishHordeArtifacts) and !$(RetrieveUsingCloudArtifacts)">
+										<Property Name="HordeArtifactPlatformRetrieveWithConfig" Value="$(HordeArtifactPlatformRetrieve) $(TestConfiguration)$(TestConfigTitle) (Horde)"/>
+										<Node Name="$(HordeArtifactPlatformRetrieveWithConfig)" Requires="$(TestPrerequisites)">
+											<Switch>
+												<Case If="$(UsePackaged)">
+													<RetrieveArtifact Name="$(ProjectName)-Packaged-$(StagedPlatformFolder)$(VariationTitle)-$(TestConfiguration)" Type="packaged-build" OutputDir="$(ProjectOutputDirectory)/Packages/$(StagedPlatformFolder)/$(TestConfiguration)" />
+												</Case>
+												<Default>
+													<RetrieveArtifact Name="$(ProjectName)-Staged-$(StagedPlatformFolder)$(VariationTitle)" Type="staged-build" OutputDir="$(ProjectOutputDirectory)/$(StagedPlatformFolder)" />
+												</Default>
+											</Switch>
+										</Node>
+										<Property Name="TestPrerequisites" Value="$(TestPrerequisites);$(HordeArtifactPlatformRetrieveWithConfig)" />
+									</Do>
+
+									<Do If="$(RetrieveUsingCloudArtifacts)">
+										<Property Name="CloudArtifactPlatformRetrieveWithConfig" Value="$(CloudArtifactPlatformRetrieve) $(TestConfiguration)$(TestConfigTitle) (Cloud)"/>
+										<Node Name="$(CloudArtifactPlatformRetrieveWithConfig)" Requires="$(TestPrerequisites)">
+											<Switch>
+												<Case If="$(UsePackaged)">
+													<RetrieveCloudArtifact 
+													Name="$(ProjectName)-$(TargetPlatform)$(VariationTitle)-Packaged-$(TestConfiguration)" 
+													Type="packaged-build" 
+													OutputDir="$(ProjectOutputDirectory)/Packages/$(StagedPlatformFolder)/$(TestConfiguration)" 
+													Project="$(TargetProject)" 
+													Platform="$(StagedPlatformFolder)"
+													BuildGroup="$(BuildVersion)" 
+													AllowMultipleMatches="true"
+													/>
+												</Case>
+												<Default>
+													<RetrieveCloudArtifact 
+													Name="$(ProjectName)-$(TargetPlatform)$(VariationTitle)-Staged" 
+													Type="staged-build" 
+													OutputDir="$(ProjectOutputDirectory)/$(StagedPlatformFolder)" 
+													Project="$(TargetProject)" 
+													Platform="$(StagedPlatformFolder)" 
+													BuildGroup="$(BuildVersion)" 
+													AllowMultipleMatches="true"
+													/>
+												</Default>
+											</Switch>
+										</Node>
+
+										<Property Name="TestPrerequisites" Value="$(TestPrerequisites);$(CloudArtifactPlatformRetrieveWithConfig)" />
+									</Do>
+
+									<Property Name="GauntletTestArgs" Value="$(GauntletTargetTestArgs)"/>
+									<Property Name="GauntletTestArgs" Value="$(GauntletTargetTestAltArgs)" If="ContainsItem('$(AlternateDevicePoolPlatforms)', '$(TargetPlatform)', ';')"/>
+
+									<Property Name="GauntletTestArgs" Value="$(GauntletTestArgs) -adbkeys=&quot;$(AdbKeysPath)&quot;" If="'$(AdbKeysPath)' != '' and '$(TargetPlatform)' == 'Android'"/>
+
+									<!-- Install step -->
+									<Do If="!$(IsLocalPlatform) And $(WithDeviceReservation)">
+										<Property Name="NodeAnnotations" Value=""/>
+										<Property Name="NodeAnnotations" Value="Workflow=$(AnnotationsTarget)" If="ContainsItem('$(AnnotationsPlatforms)', '$(TargetPlatform)', ';') and ContainsItem('$(AnnotationsTests)', 'Install', ';')"/>
+										<Property Name="InstallNodeName" Value="$(PreNodeName)$(TargetPlatform)$(VariationTitle) $(TestConfiguration)$(TestConfigTitle) Install Game"/>
+										<Property Name="InstallArgs" Value="-ClientCount=1 $(GauntletPreInstallArgs)"/>
+										<Do If="'$(AutomationReportOutputDirectory)' != ''">
+											<Property Name="OutputFragmentPath" Value="$(TargetPlatform)/$(TestConfiguration)/UE.InstallOnly" />
+											<Property Name="OutputReportPath" Value="$(AutomationReportOutputDirectory)/$(OutputFragmentPath)" />
+											<Property Name="GauntletReportArgs" Value="-logdir=&quot;$(OutputReportPath)&quot; -artifactname=Gauntlet" />
+										</Do>
+										<Node Name="$(InstallNodeName)" Requires="$(TestPrerequisites)" Annotations="DeviceReserve=Begin;$(NodeAnnotations)">
+											<Command Name="RunUnreal" Arguments="-project=$(TargetProject) -test=&quot;UE.InstallOnly&quot; -platform=$(TargetPlatform) -configuration=$(TestConfiguration) -target=$(TargetName) $(GauntletBuildArg) $(GauntletTestArgs) $(GauntletReportArgs) $(InstallArgs)" />
+										</Node>
+										<Property Name="TestPrerequisites" Value="$(TestPrerequisites);$(InstallNodeName)" />
+										<Do If="!ContainsItem('$(TargetPlatformsToSkipBootTests)', '$(TargetPlatform)', '+') And !ContainsItem('$(TargetPlatformsToSkipTests)', '$(TargetPlatform)', '+')">
+											<Property Name="AllPlatformInstallNodes" Value="$(AllPlatformInstallNodes);$(InstallNodeName)"/>
+										</Do>
+										<!-- Skip install for all other test node -->
+										<Property Name="GauntletTestArgs" Value="$(GauntletTestArgs) $(GauntletTestSkipInstallArgs)"/>
+									</Do>
+
+									<!-- Expand the tests here -->
+									<Expand Name="PreTargetTests"/>
+
+									<!-- Declare a boot test for the target-->
+									<Property Name="NodeAnnotations" Value=""/>
+									<Property Name="NodeAnnotations" Value="Workflow=$(AnnotationsTarget)" If="ContainsItem('$(AnnotationsPlatforms)', '$(TargetPlatform)', ';') and ContainsItem('$(AnnotationsTests)', 'BootTest', ';')"/>
+									<Property Name="BootTestName" Value="$(PreNodeName)$(TargetPlatform)$(VariationTitle) $(TestConfiguration)$(TestConfigTitle) BootTest" />
+									<Node Name="$(BootTestName)" Requires="$(TestPrerequisites)" Annotations="$(NodeAnnotations)">
+										<Do If="'$(AutomationReportOutputDirectory)' != ''">
+											<Property Name="OutputFragmentPath" Value="$(TargetPlatform)$(VariationTitle)/$(TestConfiguration)/UE.BootTest" />
+											<Property Name="OutputReportPath" Value="$(AutomationReportOutputDirectory)/$(OutputFragmentPath)" />
+											<Property Name="GauntletReportArgs" Value="-logdir=&quot;$(OutputReportPath)&quot; -artifactname=Gauntlet" />
+										</Do>
+										<!-- Run the boot test -->
+										<Command Name="RunUnreal" Arguments="-project=$(TargetProject) -test=&quot;UE.BootTest&quot; -platform=$(TargetPlatform) -configuration=$(TestConfiguration) -target=$(TargetName) $(GauntletBuildArg) $(GauntletTestArgs) $(GauntletReportArgs) $(TestConfigArgs)" />
+									</Node>
+
+									<!-- Link the boot test in to the prerequisites and generated nodes if its enabled-->
+									<Do If="$(TargetBootTest) and !ContainsItem('$(TargetPlatformsToSkipBootTests)', '$(TargetPlatform)', '+')">
+										<Property Name="TestPrerequisites" Value="$(TestPrerequisites);$(BootTestName)" />
+										<Property Name="AllPlatformBootTestNodes" Value="$(AllPlatformBootTestNodes);$(BootTestName)"/>
+									</Do>
+
+									<!-- Go through the list of tests -->
+									<ForEach Name="TestName" Values="$(TargetTestList)" Separator="+">
+										<Property Name="TestNameTitle" Value="$(TestName)"/>
+										<Expand Name="SanitizeTestNameTitle"/>
+										<!-- Run and report on the test-->
+										<Do If="'$(AutomationReportOutputDirectory)' != ''">
+											<Property Name="OutputFragmentPath" Value="$(TargetPlatform)$(VariationTitle)/$(TestConfiguration)/$(TestNameTitle)" />
+											<Property Name="OutputReportPath" Value="$(AutomationReportOutputDirectory)/$(OutputFragmentPath)" />
+											<Property Name="GauntletReportArgs" Value="-logdir=&quot;$(OutputReportPath)&quot; -artifactname=Gauntlet -ReportExportPath=&quot;$(OutputReportPath)&quot;" />
+											<Property Name="GauntletReportArgs" Value="$(GauntletReportArgs) -ReportURL=&quot;$(AutomationReportUrl)/$(OutputFragmentPath)/&quot;" If="'$(AutomationReportUrl)' != ''" />
+										</Do>
+										<Property Name="TestNodeName" Value="$(PreNodeName)$(TargetPlatform)$(VariationTitle) $(TestConfiguration)$(TestConfigTitle) $(TestNameTitle)" />
+										<Property Name="NodeAnnotations" Value=""/>
+										<Property Name="NodeAnnotations" Value="Workflow=$(AnnotationsTarget)" If="ContainsItem('$(AnnotationsPlatforms)', '$(TargetPlatform)', ';') and ContainsItem('$(AnnotationsTests)', '$(TestName)', ';')"/>
+
+										<Property Name="TestNameBeforeParentheses" Value="$(TestName)"/>
+										<Property Name="TestArgumentsWithinParentheses" Value=""/>
+										<Expand Name="SplitCombinedTestNameAndArgs"/>
+
+										<!-- Test 'Requires' should be the package game and the publish folder is used; however some tests may require the editor ie: UE.TargetAutomation -->
+										<Property Name="AdditionalTestPrerequisites" Value=""/>
+										<ForEach Name="ToMatchName" Values="$(RequiringEditorTargetTestList)" Separator="+">
+											<Do If="Contains('$(TestNameBeforeParentheses)', '$(ToMatchName)') And !ContainsItem('$(TestPrerequisites)', '$(PlatformEditorNodeName)', ';')">
+												<Property Name="AdditionalTestPrerequisites" Value="$(AdditionalTestPrerequisites);$(PlatformEditorNodeName)" />
+											</Do>
+										</ForEach>
+
+										<!-- 'Requires' should be the package game and the publish folder is used; however some tests may require the editor ie: UE.TargetAutomation -->
+										<Node Name="$(TestNodeName)" Requires="$(TestPrerequisites);$(AdditionalTestPrerequisites)" Annotations="$(NodeAnnotations)">
+											<Command Name="RunUnreal" Arguments="-project=$(TargetProject) -test=&quot;$(TestNameBeforeParentheses)&quot; -platform=$(TargetPlatform) -configuration=$(TestConfiguration) -target=$(TargetName) $(GauntletBuildArg) $(GauntletTestArgs) $(TestArgumentsWithinParentheses) $(GauntletReportArgs) $(TestConfigArgs)"/>
+										</Node>
+										<!-- Add it to our list -->
+										<Do If="!ContainsItem('$(TargetPlatformsToSkipTests)', '$(TargetPlatform)', '+') And (!Contains('$(TestListToIncludePlatforms)', '^$(TestNameBeforeParentheses)|') Or ContainsItem('$(TestListToIncludePlatforms)', '$(TestNameBeforeParentheses)|$(TargetPlatform)', '^'))">
+											<Property Name="AllPlatformTestNodes" Value="$(AllPlatformTestNodes);$(TestNodeName)"/>
+										</Do>
+									</ForEach>
+
+									<!-- Expand the tests here -->
+									<Expand Name="AdditionalTestNodes"/>
+								</ForEach>
+							</Agent>
+						</Do>
 					</ForEach>
-				</Agent>
+				</Do>
+				
+				<!-- Declare labels for CIS -->
+				<Label Category="Clients" Name="$(PreLabelName)$(TargetPlatform)$(VariationTitle)" Requires="$(PlatformCompileNodeName);$(UploadSymbolsNodeName)" />
+				<Label Category="Clients" Name="$(PreLabelName)Content $(TargetPlatform)$(VariationTitle)" Requires="$(PlatformCookNodeName);$(PlatformStageNodeName);$(PlatformPackageNodeName)" />
+				<Label Category="Clients" Name="$(PreLabelName)Publish $(TargetPlatform)$(VariationTitle)" Requires="$(PlatformPublishPackagedNodeName);$(PlatformPublishPackagedNodeName);$(PlatformPublishStagedNodeName);$(PlatformHashNodeName)" />
+				<Label Category="Servers" Name="$(PreLabelName)$(TargetPlatform)$(VariationTitle) Server" Requires="$(PlatformCompileNodeName) Server" />
+				<Label Category="Servers" Name="$(PreLabelName)Content $(TargetPlatform)$(VariationTitle) Server" Requires="$(PlatformCookNodeName) Server;$(PlatformStageNodeName) Server;$(PlatformPackageNodeName) Server" />
+				<Label Category="Servers" Name="$(PreLabelName)Publish $(TargetPlatform)$(VariationTitle) Server " Requires="$(PlatformPublishPackagedNodeName) Server;$(PlatformPublishStagedNodeName) Server" />
+				<Label Category="Tests" Name="$(PreLabelName)Install $(TargetPlatform)$(VariationTitle)" Requires="$(AllPlatformInstallNodes)" />
+				<Label Category="Tests" Name="$(PreLabelName)BootTest $(TargetPlatform)$(VariationTitle)" Requires="$(AllPlatformBootTestNodes)" />
+				<Label Category="Tests" Name="$(PreLabelName)Test $(TargetPlatform)$(VariationTitle)" Requires="$(AllPlatformTestNodes)" />
+
+				<!-- Tie all steps into our platform tasks unless excluded -->
+				<Property Name="PlatformTasks" Value="" />
+				<Property Name="PlatformTasks" Value="$(PlatformCompileNodeName);$(UploadSymbolsNodeName)" If="!$(SkipBuild)" />
+				<Property Name="PlatformTasks" Value="$(PlatformCompileNodeName) Server" If="!$(SkipBuild) and $(DedicatedServer)" />
+				<Property Name="PlatformTasks" Value="$(PlatformTasks);$(PlatformCookNodeName)" If="!$(SkipCook) and !$(SkipPublish)" />
+				<Property Name="PlatformTasks" Value="$(PlatformTasks);$(PlatformCookNodeName) Server" If="!$(SkipCook) and !$(SkipPublish) and $(DedicatedServer)" />
+				<Property Name="PlatformTasks" Value="$(PlatformTasks);$(RequirePlatformStaged)" If="!$(SkipPublish)" />
+				<Property Name="PlatformTasks" Value="$(PlatformTasks);$(RequirePlatformPackaged)" If="!$(SkipPackage) and !$(SkipPublish)" />
+				<Property Name="PlatformTasks" Value="$(PlatformTasks);$(PlatformHashNodeName)"  If="$(PublishHashedStagingData)" />	
+				<Property Name="PlatformTasks" Value="$(PlatformTasks);$(AllPlatformBootTestNodes);$(AllPlatformTestNodes)" />
+				<Property Name="PlatformTasks" Value="$(PlatformTasks);$(PlatformCheckBuildSizeNodeName)" If="'$(BuildSizeLimit)' != ''" />
+
+				<Property Name="GeneratedTargetNodes" Value="$(GeneratedTargetNodes);$(PlatformTasks)"/>
+				<Property Name="GeneratedTargetTestNodes" Value="$(GeneratedTargetTestNodes);$(AllPlatformBootTestNodes);$(AllPlatformTestNodes)"/>
 			</Do>
 		</ForEach>
-		
-		<!-- Declare labels for CIS -->
-		<Label Category="Clients" Name="$(PreLabelName)$(TargetPlatform)" Requires="$(PlatformCompileNodeName);$(UploadSymbolsNodeName)" />
-		<Label Category="Clients" Name="$(PreLabelName)Content $(TargetPlatform)" Requires="$(PlatformCookNodeName);$(PlatformStageNodeName);$(PlatformPackageNodeName)" />
-		<Label Category="Clients" Name="$(PreLabelName)Publish $(TargetPlatform)" Requires="$(PlatformPublishPackagedNodeName);$(PlatformPublishPackagedNodeName);$(PlatformPublishStagedNodeName);$(PlatformHashNodeName)" />
-		<Label Category="Servers" Name="$(PreLabelName)$(TargetPlatform) Server" Requires="$(PlatformCompileNodeName) Server" />
-		<Label Category="Servers" Name="$(PreLabelName)Content $(TargetPlatform) Server" Requires="$(PlatformCookNodeName) Server;$(PlatformStageNodeName) Server;$(PlatformPackageNodeName) Server" />
-		<Label Category="Servers" Name="$(PreLabelName)Publish $(TargetPlatform) Server " Requires="$(PlatformPublishPackagedNodeName) Server;$(PlatformPublishStagedNodeName) Server" />
-		<Label Category="Tests" Name="$(PreLabelName)Install $(TargetPlatform)" Requires="$(AllPlatformInstallNodes)" />
-		<Label Category="Tests" Name="$(PreLabelName)BootTest $(TargetPlatform)" Requires="$(AllPlatformBootTestNodes)" />
-		<Label Category="Tests" Name="$(PreLabelName)Test $(TargetPlatform)" Requires="$(AllPlatformTestNodes)" />
-
-		<!-- Tie all steps into our platform tasks unless excluded -->
-		<Property Name="PlatformTasks" Value="" />
-		<Property Name="PlatformTasks" Value="$(PlatformCompileNodeName);$(UploadSymbolsNodeName)" If="!$(SkipBuild)" />
-		<Property Name="PlatformTasks" Value="$(PlatformCompileNodeName) Server" If="!$(SkipBuild) and $(DedicatedServer)" />
-		<Property Name="PlatformTasks" Value="$(PlatformTasks);$(PlatformCookNodeName)" If="!$(SkipCook) and !$(SkipPublish)" />
-		<Property Name="PlatformTasks" Value="$(PlatformTasks);$(PlatformCookNodeName) Server" If="!$(SkipCook) and !$(SkipPublish) and $(DedicatedServer)" />
-		<Property Name="PlatformTasks" Value="$(PlatformTasks);$(RequirePlatformStaged)" If="!$(SkipPublish)" />
-		<Property Name="PlatformTasks" Value="$(PlatformTasks);$(RequirePlatformPackaged)" If="!$(SkipPackage) and !$(SkipPublish)" />
-		<Property Name="PlatformTasks" Value="$(PlatformTasks);$(PlatformHashNodeName)"  If="$(PublishHashedStagingData)" />	
-		<Property Name="PlatformTasks" Value="$(PlatformTasks);$(AllPlatformBootTestNodes);$(AllPlatformTestNodes)" />
-		<Property Name="PlatformTasks" Value="$(PlatformTasks);$(PlatformCheckBuildSizeNodeName)" If="'$(BuildSizeLimit)' != ''" />
-
-		<Property Name="GeneratedTargetNodes" Value="$(GeneratedTargetNodes);$(PlatformTasks)"/>
-		<Property Name="GeneratedTargetTestNodes" Value="$(GeneratedTargetTestNodes);$(AllPlatformBootTestNodes);$(AllPlatformTestNodes)"/>
 	</ForEach>
 
 	<!-- Create an aggregate with all the generated nodes -->
diff --git a/Engine/Build/Graph/Tasks/Inc/PlatformOptions.xml b/Engine/Build/Graph/Tasks/Inc/PlatformOptions.xml
--- a/Engine/Build/Graph/Tasks/Inc/PlatformOptions.xml
+++ b/Engine/Build/Graph/Tasks/Inc/PlatformOptions.xml
@@ -34,6 +34,10 @@
 <Option Name="TargetPlatforms" DefaultValue="" Description="Alternative to 'WithFoo', a list of Platforms to build and test. Android+IOS etc"/>
 <Option Name="CookedEditorPlatforms" DefaultValue="" Description="Alternative to 'WithFooCookedEditor', a list of editors to build and cooked. Win64+Mac etc"/>
 
+<Option Name="Win64Variations" DefaultValue="Default+arm64+arm64ec+MultiArch" Description="List of variation to build on Win64"/>
+<Option Name="LinuxVariations" DefaultValue="Default" Description="List of variation to build on Linux"/>
+<Option Name="MacVariations" DefaultValue="Default" Description="List of variation to build on Mac"/>
+
 <Option Name="EditorPlatformWin64AgentOverride" DefaultValue="Win64" Restrict="Win64|Linux|Mac" Description="Platform to use to build an editor for Win64 target"/>
 <Option Name="EditorPlatformLinuxAgentOverride" DefaultValue="Linux" Restrict="Win64|Linux|Mac" Description="Platform to use to build an editor for Linux target"/>
 <Option Name="EditorPlatformMacAgentOverride" DefaultValue="Mac" Restrict="Win64|Linux|Mac" Description="Platform to use to build an editor for Mac target"/>
@@ -49,6 +53,7 @@
 <Option Name="TargetPlatformsToSkipTests" DefaultValue="IOS+TVOS" Restrict=".*" Description="List of target platforms that target tests won't be run on" />
 <Option Name="CookedEditorPlatformsToSkipTests" DefaultValue="" Restrict=".*" Description="List of editor cooked platforms that editor tests won't be run on" />
 <Option Name="PlatformsToSkipTests" DefaultValue="" Restrict=".*" Description="List of platforms that all tests won't be run on. Overrides Editor/TargetPlatformsToSkipTests when non-empty." />
+<Option Name="VariationsToSkipTests" DefaultValue="arm64+arm64ec" Description="List of variations that all tests won't be run with."/>
 
 <!-- Packaging isn't supported by Gauntlet yet :( )-->
 <Property Name="PlatformsToSkipPackaging" Value="Win64+Linux" />
@@ -84,6 +89,17 @@
 <Option Name="WithTVOS" DefaultValue="false" Description="Build and test TVOS"/>
 <Option Name="WithWin64" DefaultValue="false" Description="Build and test Win64"/>
 
+<Option Name="WithDefaultVariation" DefaultValue="true" Description="Build and test Default Variations"/>
+<Option Name="WithArm64Variation" DefaultValue="false" Description="Build and test arm64 Variations"/>
+<Option Name="WithArm64ecVariation" DefaultValue="false" Description="Build and test arm64ec Variations"/>
+<Option Name="WithMultiArchVariation" DefaultValue="false" Description="Build and test x64+arm64 Variations"/>
+
+<Property Name="RequestedVariations" Value=""/>
+<Property Name="RequestedVariations" Value="Default" If="$(WithDefaultVariation)"/>
+<Property Name="RequestedVariations" Value="$(RequestedVariations)+arm64" If="$(WithArm64Variation)"/>
+<Property Name="RequestedVariations" Value="$(RequestedVariations)+arm64ec" If="$(WithArm64ecVariation)"/>
+<Property Name="RequestedVariations" Value="$(RequestedVariations)+MultiArch" If="$(WithMultiArchVariation)"/>
+
 <Option Name="AndroidCookFlavor" DefaultValue="ETC2" Description="Cook flavor to use for build/cook/run when target platform is Android"/>
 
 <Do If="'$(EditorPlatforms)' == ''" >
