Allow the DDCVerify job to store archive with all *.fromcache and *.verify files produced during the run in an optionally specified storage path.

#Synced-CL 50932609

diff --git a/Engine/Build/Graph/Tests/DDCVerify.xml b/Engine/Build/Graph/Tests/DDCVerify.xml
--- a/Engine/Build/Graph/Tests/DDCVerify.xml
+++ b/Engine/Build/Graph/Tests/DDCVerify.xml
@@ -3,6 +3,7 @@
 
 	<!-- Project Options -->
 	<Option Name="ProjectFile" DefaultValue="_Unknown_" Restrict=".+" Description="The projects to cook"/>
+	<Option Name="ProjectPath" DefaultValue="_Unknown_" Restrict=".+" Description="The path containing the project file to cook"/>
 	<Option Name="ProjectName" DefaultValue="_Unknown_" Restrict=".+" Description="The projects to cook"/>
 	<Option Name="EditorTarget" DefaultValue="$(ProjectName)Editor" Description="Compile name of the editor target"/>
 	<Option Name="DDCGraph" DefaultValue="" Description="Overridable DDC graph to use when performing the verification (default: empty, uses default graph)"/>
@@ -25,10 +26,13 @@
 	<!-- DDC Verify Options -->
 	<Option Name="DDCVerifyDefaultRate" Restrict="\d*" DefaultValue="1" Description="The default verification rate in percent"/>
 	<Option Name="DDCVerifyCustomBucketRates" DefaultValue="" Description="Bucket names and associated custom verification rates that differ from the default rate.  Uses plus or comma as separators."/>
+	<Option Name="DDCVerifyArtifactStoragePath" DefaultValue="" Description="Folder to store verify artifacts as a zip archive."/>
+
 	<Property Name="DDCVerifyArgs" Value="-DDC-Verify"/>
 	<Property Name="DDCVerifyArgs" Value="$(DDCVerifyArgs)=$(DDCVerifyCustomBucketRates)" If="'$(DDCVerifyCustomBucketRates)' != ''"/>
 	<Property Name="DDCVerifyArgs" Value="$(DDCVerifyArgs) -DDC-VerifyRate=$(DDCVerifyDefaultRate)"/>
 
+
 	<Property Name="DDCGraphArgs" Value=""/>
 	<Property Name="DDCGraphArgs" Value="-ddc=$(DDCGraph)" If="'$(DDCGraph)' != ''"/>
 
@@ -51,8 +55,15 @@
 		<Node Name="DDC Verify" Requires="$(CompileNodeName)">
 			<Commandlet Name="Cook" Project="$(ProjectFile)" Arguments="$(CommonCookArgs)"/>
 		</Node>
+
+		<Node Name="DDC Verify Archive Storage">
+			<Zip FromDir="$(ProjectPath)/Saved/VerifyDDC" Files="*.*" ZipFile="$(ProjectPath)/Saved/VerifyDDC.zip"/>
+			<Copy Files="VerifyDDC.zip" From="$(ProjectPath)/Saved/" To="$(DDCVerifyArtifactStoragePath)"/>
+		</Node>
 	</Agent>
 
-	<Aggregate Name="DDCVerify" Requires="DDC Verify"/>
+	<Property Name="DDCVerifyTargets" Value="DDC Verify"/>
+	<Property Name="DDCVerifyTargets" Value="$(DDCVerifyTargets);DDC Verify Archive Storage" If="'$(DDCVerifyArtifactStoragePath)' != ''"/>
+	<Aggregate Name="DDCVerify" Requires="$(DDCVerifyTargets)"/>
 
 </BuildGraph>
\ No newline at end of file
