Retrieve Dedicated server build artifacts when running test if DedicatedServer was enabled

#rnx
#Synced-CL 52205294

diff --git a/Engine/Build/Graph/Tasks/BuildAndTestProject.xml b/Engine/Build/Graph/Tasks/BuildAndTestProject.xml
--- a/Engine/Build/Graph/Tasks/BuildAndTestProject.xml
+++ b/Engine/Build/Graph/Tasks/BuildAndTestProject.xml
@@ -731,14 +731,14 @@
 				<Property Name="RequirePlatformStaged" Value="$(PlatformStageNodeName)" />
 				<Property Name="RequirePlatformStaged" Value="$(PlatformStageNodeName) Server" If="$(DedicatedServer)"/>
 				<Property Name="RequirePlatformStaged" Value="$(RequirePlatformStaged);$(PlatformStagedArtifactNodeName)" If="!$(SkipPublish) And $(PublishHordeArtifacts)"/>
-				<Property Name="RequirePlatformStaged" Value="$(RequirePlatformStaged);$(PlatformStagedArtifactNodeName) Server" If="!$(SkipPublish) And $(PublishHordeArtifacts) And $(DedicatedServer)"/>
+				<Property Name="RequirePlatformStaged" Value="$(RequirePlatformStaged);$(PlatformStagedArtifactNodeName) Server" If="!$(SkipPublish) And $(PublishHordeArtifacts) And $(DedicatedServer) And $(ServerPackagedSeparately) And $(IsDesktopPlatform)"/>
 				<Property Name="RequirePlatformStaged" Value="$(RequirePlatformStaged);$(CloudPlatformStagedArtifactNodeName)" If="!$(SkipPublish) And $(PublishCloudArtifacts)"/>
 				<Property Name="RequirePlatformStaged" Value="$(RequirePlatformStaged);$(PlatformPublishStagedNodeName)" If="!$(SkipPublish) And $(PublishBuilds)"/>
 				<Property Name="RequirePlatformStaged" Value="$(RequirePlatformStaged);$(PlatformPublishStagedNodeName) Server" If="!$(SkipPublish) And $(PublishBuilds) And $(DedicatedServer)"/>
 				<Property Name="RequirePlatformPackaged" Value="$(PlatformPackageNodeName)"/>
 				<Property Name="RequirePlatformPackaged" Value="$(PlatformPackageNodeName) Server" If="$(DedicatedServer)"/>
 				<Property Name="RequirePlatformPackaged" Value="$(RequirePlatformPackaged);$(HordeArtifactPlatformPackagedNodeName)" If="!$(SkipPublish) And $(PublishHordeArtifacts)"/>
-				<Property Name="RequirePlatformPackaged" Value="$(RequirePlatformPackaged);$(HordeArtifactPlatformPackagedNodeName) Server" If="!$(SkipPublish) And $(PublishHordeArtifacts) And $(DedicatedServer)"/>
+				<Property Name="RequirePlatformPackaged" Value="$(RequirePlatformPackaged);$(HordeArtifactPlatformPackagedNodeName) Server" If="!$(SkipPublish) And $(PublishHordeArtifacts) And $(DedicatedServer) And $(ServerPackagedSeparately) And $(IsDesktopPlatform)"/>
 				<Property Name="RequirePlatformPackaged" Value="$(RequirePlatformPackaged);$(CloudArtifactPlatformPackagedNodeName)" If="!$(SkipPublish) And $(PublishCloudArtifacts)"/>
 				<Property Name="RequirePlatformPackaged" Value="$(RequirePlatformPackaged);$(PlatformPublishPackagedNodeName)"  If="!$(SkipPublish) And $(PublishBuilds)"/>
 				<Property Name="RequirePlatformPackaged" Value="$(RequirePlatformPackaged);$(PlatformPublishPackagedNodeName) Server"  If="!$(SkipPublish) And $(PublishBuilds) And $(DedicatedServer)"/>
@@ -753,7 +753,10 @@
 				<Property Name="BCRCookFlavor" Value=""/>
 				<Property Name="BCRCookFlavor" Value="$(AndroidCookFlavor)" If="'$(TargetPlatform)' == 'Android'" />
 				
-				<Property Name="StagedPlatformFolder" Value="$(CookPlatform)"  />
+				<Property Name="StagedPlatformFolder" Value="$(CookPlatform)" />
+				<Property Name="ServerFolder" Value="WindowsServer"/>
+				<Property Name="ServerFolder" Value="LinuxServer" If="$(IsTargetPlatformLinux)"/>
+				<Property Name="ServerFolder" Value="MacServer" If="$(IsTargetPlatformMac)"/>
 				<!-- <Property Name="StagedPlatformFolder" Value="$(StagedPlatformFolder)Client" If="'$(TargetType)' == 'Client'" /> -->
 
 				<!-- Hashing of staged data can only be done on Win64 hosts -->
@@ -853,12 +856,7 @@
 							<Expand Name="PreCookStep"/>
 							<Property Name="CookArgs" Value="" />
 							<Property Name="CookArgs" Value="-MapIniSection=&quot;$(MapIniSection)&quot; $(CookPlatform)" If="'$(MapIniSection)' != ''"/>	
-
-							<Property Name="ServerPlatformCook" Value="Windows"/>
-							<Property Name="ServerPlatformCook" Value="Linux" If="$(IsTargetPlatformLinux)"/>
-							<Property Name="ServerPlatformCook" Value="$(ServerPlatformCook)Server"/>
-
-							<Cook Project="$(TargetProject)" Platform="$(ServerPlatformCook)" TagOutput="false" Arguments="$(CookArgs) $(ExtraCookArguments)"/>
+							<Cook Project="$(TargetProject)" Platform="$(ServerFolder)" TagOutput="false" Arguments="$(CookArgs) $(ExtraCookArguments)"/>
 						</Do>
 					</Node>
 					
@@ -881,6 +879,9 @@
 						<Property Name="BCRArgsServer" Value="$(BCRArgs) -server" />
 						<Property Name="ServerStageArguments" Value="$(BCRArgsServer) -configuration=$(PackageConfigurations) $(StageAndPackageServerArguments) -skipbuild -skipcook -stage -pak -NoCodeSign" />
 						<Command Name="BuildCookRun" Arguments="$(ServerStageArguments)" />
+						<Do If="!$(ServerPackagedSeparately)">
+							<Copy Files="..." From="$(ProjectOutputDirectory)/$(ServerFolder)" To="$(ProjectOutputDirectory)/$(StagedPlatformFolder)/$(ServerFolder)"/>
+						</Do>
 					</Node>
 					
 					<!-- Hash -->
@@ -954,8 +955,8 @@
 						<Log Message="Network upload is disabled, skipping upload." If="!$(PublishToNetworkOutputDirectory)"/>
 						<Do If="$(PublishToNetworkOutputDirectory)">				
 							<Error Message="NetworkOutputDirectory is not set" If="'$(NetworkOutputDirectory)' == ''"/>
-							<Copy Files="..." From="$(ProjectOutputDirectory)/$(StagedPlatformFolder)Server" To="$(NetworkOutputDirectory)/$(StagedPlatformFolder)Server/Staged"/>
-							<Copy Files="..." From="$(ShaderSymbolsDirectory)/*$(TargetPlatform)Server*" To="$(NetworkOutputDirectory)/$(StagedPlatformFolder)Server/ShaderSymbols"/>
+							<Copy Files="..." From="$(ProjectOutputDirectory)/$(ServerFolder)" To="$(NetworkOutputDirectory)/$(ServerFolder)/Staged"/>
+							<Copy Files="..." From="$(ShaderSymbolsDirectory)/*$(TargetPlatform)Server*" To="$(NetworkOutputDirectory)/$(ServerFolder)/ShaderSymbols"/>
 						</Do>
 					</Node>
 
@@ -993,13 +994,12 @@
 					
 					<Do If="$(PublishHordeArtifacts)">
 						<Node Name="$(PlatformStagedArtifactNodeName) Server" Requires="$(PlatformStageNodeName)">
-							<Property Name="ServerArtifactsPath" Value="$(ProjectOutputDirectory)/$(StagedPlatformFolder)" />
-							<Property Name="ServerArtifactsPath" Value="$(ProjectOutputDirectory)/$(StagedPlatformFolder)Server" If="$(ServerPackagedSeparately)" />
+							<Property Name="ServerArtifactsPath" Value="$(ProjectOutputDirectory)/$(ServerFolder)"/>
 							
 							<Expand Name="WriteArtifactMetadataFile" Type="staged-build" Platform="$(StagedPlatformFolder)$(VariationTitle)" Configuration="$(PackageConfigurations)" ArtifactDir="$(ServerArtifactsPath)"/>
-							<Tag Files="$(ServerArtifactsPath)/..." Except="$(TargetPlatform)Server/..." With="#$(PlatformStageNodeName) Server Files"/>
+							<Tag Files="$(ServerArtifactsPath)/..." With="#$(PlatformStageNodeName) Server Files"/>
 
-							<CreateArtifact Name="$(ProjectName)-Server-Staged-$(StagedPlatformFolder)$(VariationTitle)" Type="staged-build" BaseDir="$(ServerArtifactsPath)" Keys="Project=$(ProjectName)" Metadata="Stream=$(Branch);Platform=$(StagedPlatformFolder);BuildName=$(BuildNamePath);dashboard-category=Staged Builds" Files="#$(PlatformStageNodeName) Server Files"/>
+							<CreateArtifact Name="$(ProjectName)-Staged-$(ServerFolder)$(VariationTitle)" Type="staged-build" BaseDir="$(ServerArtifactsPath)" Keys="Project=$(ProjectName)" Metadata="Stream=$(Branch);Platform=$(ServerFolder);BuildName=$(BuildNamePath);dashboard-category=Staged Builds" Files="#$(PlatformStageNodeName) Server Files"/>
 						</Node>
 					</Do>
 
@@ -1014,7 +1014,7 @@
 					<Node Name="$(PlatformPackageNodeName) Server" Requires="$(PlatformPublishStageDependencyName)">
 						<ForEach Name="PackageConfiguration" Values="$(PackageConfigurations)" Separator="+">
 							<Property Name="ServerArchiveDirectory" Value="$(ProjectOutputDirectory)/Packages/$(StagedPlatformFolder)$(VariationTitle)/$(PackageConfiguration)" />
-							<Property Name="ServerArchiveDirectory" Value="$(ProjectOutputDirectory)/Packages/$(TargetPlatform)Server$(VariationTitle)/$(PackageConfiguration)" If="$(ServerPackagedSeparately)" />
+							<Property Name="ServerArchiveDirectory" Value="$(ProjectOutputDirectory)/Packages/$(ServerFolder)$(VariationTitle)/$(PackageConfiguration)" If="$(ServerPackagedSeparately)" />
 							<Property Name="BCRArgsServer" Value="$(BCRArgs) -server" />	
 						<Command Name="BuildCookRun" Arguments="$(BCRArgsServer) -configuration=$(PackageConfiguration) $(StageAndPackageServerArguments) -skipbuild -skipcook -skipstage -skippak -package $(PackageOnlyArguments) -archive -archivedirectory=&quot;$(ServerArchiveDirectory)&quot;" />
 						</ForEach>
@@ -1038,7 +1038,7 @@
 							<Error Message="NetworkOutputDirectory is not set" If="'$(NetworkOutputDirectory)' == ''"/>
 							<ForEach Name="PackageConfiguration" Values="$(PackageConfigurations)" Separator="+">
 								<Property Name="ServerArchiveDirectory" Value="$(NetworkOutputDirectory)/Packages/$(StagedPlatformFolder)$(VariationTitle)"/>
-								<Property Name="ServerArchiveDirectory" Value="$(NetworkOutputDirectory)/Packages/$(TargetPlatform)Server$(VariationTitle)" If="$(ServerPackagedSeparately)"/>
+								<Property Name="ServerArchiveDirectory" Value="$(NetworkOutputDirectory)/Packages/$(ServerFolder)$(VariationTitle)" If="$(ServerPackagedSeparately)"/>
 								<Property Name="BCRArgsServer" Value="$(BCRArgs) -server" />	
 								<Command Name="BuildCookRun" Arguments="$(BCRArgsServer) -configuration=$(PackageConfiguration) $(StageAndPackageServerArguments) -skipbuild -skipcook -skipstage -skippak -skippackage -archive -archivedirectory=&quot;$(ServerArchiveDirectory)&quot;" />
 							</ForEach>
@@ -1084,12 +1084,11 @@
 					
 						<Node Name="$(HordeArtifactPlatformPackagedNodeName) Server" Requires="$(PlatformPackageNodeName)">
 							<ForEach Name="PackageConfiguration" Values="$(PackageConfigurations)" Separator="+">
-								<Property Name="ServerArchivePath" Value="$(ProjectOutputDirectory)/Packages/$(StagedPlatformFolder)$(VariationTitle)/$(PackageConfiguration)" />
-								<Property Name="ServerArchivePath" Value="$(ProjectOutputDirectory)/Packages/$(StagedPlatformFolder)Server$(VariationTitle)/$(PackageConfiguration)" If="$(ServerPackagedSeparately)" />
+								<Property Name="ServerArchivePath" Value="$(ProjectOutputDirectory)/Packages/$(ServerFolder)$(VariationTitle)/$(PackageConfiguration)"/>
 								<Property Name="BCRArgsServer" Value="$(BCRArgs) -server" />
 								<Command Name="BuildCookRun" Arguments="$(BCRArgsServer) -configuration=$(PackageConfiguration) $(StageAndPackageServerArguments) -skipbuild -skipcook -skippak $(HordePackageSkipBCRArguments) -archive -archivedirectory=&quot;$(ServerArchivePath)&quot;" />
-								<Tag Files="$(ServerArchivePath)/..."  Except="$(TargetPlatform)Server/..." With="#$(PlatformPackageNodeName) $(PackageConfiguration) Server Files"/>
-								<CreateArtifact Name="$(ProjectName)-Server-Packaged-$(TargetPlatform)$(VariationTitle)-$(PackageConfiguration)" Type="packaged-build" BaseDir="$(ServerArchivePath)" Keys="Project=$(ProjectName)" Metadata="Stream=$(Branch);Platform=$(TargetPlatform);Configuration=$(PackageConfiguration);BuildName=$(BuildNamePath);dashboard-category=$(PackageConfiguration) Builds" Files="#$(PlatformPackageNodeName) $(PackageConfiguration) Server Files"/>
+								<Tag Files="$(ServerArchivePath)/..." With="#$(PlatformPackageNodeName) $(PackageConfiguration) Server Files"/>
+								<CreateArtifact Name="$(ProjectName)-Packaged-$(ServerFolder)$(VariationTitle)-$(PackageConfiguration)" Type="packaged-build" BaseDir="$(ServerArchivePath)" Keys="Project=$(ProjectName)" Metadata="Stream=$(Branch);Platform=$(ServerFolder);Configuration=$(PackageConfiguration);BuildName=$(BuildNamePath);dashboard-category=$(PackageConfiguration) Builds" Files="#$(PlatformPackageNodeName) $(PackageConfiguration) Server Files"/>
 							</ForEach>
 						</Node>
 					</Do>
@@ -1162,9 +1161,11 @@
 											<Switch>
 												<Case If="$(UsePackaged)">
 													<RetrieveArtifact Name="$(ProjectName)-Packaged-$(StagedPlatformFolder)$(VariationTitle)-$(TestConfiguration)" Type="packaged-build" OutputDir="$(ProjectOutputDirectory)/Packages/$(StagedPlatformFolder)/$(TestConfiguration)" />
+													<RetrieveArtifact Name="$(ProjectName)-Packaged-$(ServerFolder)$(VariationTitle)-$(TestConfiguration)" Type="packaged-build" OutputDir="$(ProjectOutputDirectory)/Packages/$(StagedPlatformFolder)/$(TestConfiguration)/$(ServerFolder)" If="$(ServerPackagedSeparately)"/>
 												</Case>
 												<Default>
 													<RetrieveArtifact Name="$(ProjectName)-Staged-$(StagedPlatformFolder)$(VariationTitle)" Type="staged-build" OutputDir="$(ProjectOutputDirectory)/$(StagedPlatformFolder)" />
+													<RetrieveArtifact Name="$(ProjectName)-Staged-$(ServerFolder)$(VariationTitle)" Type="staged-build" OutputDir="$(ProjectOutputDirectory)/$(StagedPlatformFolder)/$(ServerFolder)" If="$(ServerPackagedSeparately)"/>
 												</Default>
 											</Switch>
 										</Node>
